EXR: Fixed static analysis error.

diff --git a/Engine/Source/Runtime/ImageWrapper/Private/Formats/ExrImageWrapper.h b/Engine/Source/Runtime/ImageWrapper/Private/Formats/ExrImageWrapper.h
--- a/Engine/Source/Runtime/ImageWrapper/Private/Formats/ExrImageWrapper.h
+++ b/Engine/Source/Runtime/ImageWrapper/Private/Formats/ExrImageWrapper.h
@@ -15,7 +15,7 @@ THIRD_PARTY_INCLUDES_START
 	#include "OpenEXR/include/ImfInputFile.h"
 	#include "OpenEXR/include/ImfOutputFile.h"
 	#include "OpenEXR/include/ImfArray.h"
-	#include "OpenEXR/include//ImfHeader.h"
+	#include "OpenEXR/include/ImfHeader.h"
 	#include "OpenEXR/include/ImfStdIO.h"
 	#include "OpenEXR/include/ImfChannelList.h"
 	#include "OpenEXR/include/ImfRgbaFile.h"
