Fixes a memory overrite when saving tiny EXRs

diff --git a/Engine/Source/Runtime/ImageWrapper/Private/Formats/ExrImageWrapper.cpp b/Engine/Source/Runtime/ImageWrapper/Private/Formats/ExrImageWrapper.cpp
--- a/Engine/Source/Runtime/ImageWrapper/Private/Formats/ExrImageWrapper.cpp
+++ b/Engine/Source/Runtime/ImageWrapper/Private/Formats/ExrImageWrapper.cpp
@@ -60,6 +60,12 @@ public:
 
 	virtual void write(const char c[/*n*/], int n)
 	{
+		int32 DestPost = Pos + n;
+		if (DestPost > Data.Num())
+		{
+			Data.AddUninitialized(FMath::Max(Data.Num() * 2, DestPost) - Data.Num());
+		}
+
 		for (int32 i = 0; i < n; ++i)
 		{
 			Data[Pos + i] = c[i];
@@ -314,8 +320,6 @@ void FExrImageWrapper::CompressRaw(const sourcetype* SrcData, bool bIgnoreAlpha)
 	checkSlow(OutputPixelSize >= 1 && OutputPixelSize <= 8);
 
 	FMemFileOut MemFile("");
-	MemFile.Data.AddUninitialized(Width * Height * NumWriteComponents * OutputPixelSize);
-
 	Imf::FrameBuffer ImfFrameBuffer;
 	TArray<uint8> ChannelOutputBuffers[4];
 
@@ -326,6 +330,8 @@ void FExrImageWrapper::CompressRaw(const sourcetype* SrcData, bool bIgnoreAlpha)
 
 	Imf::OutputFile ImfFile(MemFile, Header);
 	ImfFile.setFrameBuffer(ImfFrameBuffer);
+	
+	MemFile.Data.AddUninitialized(Width * Height * NumWriteComponents * (OutputFormat == 2 ? 4 : 2));
 	ImfFile.writePixels(Height);
 
 	CompressedData.AddUninitialized(MemFile.tellp());
