MediaFramework: Fixed EXR images not being written out correctly causing them to be very slow to load in.

diff --git a/Engine/Source/Runtime/ImageWrapper/Private/Formats/ExrImageWrapper.cpp b/Engine/Source/Runtime/ImageWrapper/Private/Formats/ExrImageWrapper.cpp
--- a/Engine/Source/Runtime/ImageWrapper/Private/Formats/ExrImageWrapper.cpp
+++ b/Engine/Source/Runtime/ImageWrapper/Private/Formats/ExrImageWrapper.cpp
@@ -328,14 +328,22 @@ void FExrImageWrapper::CompressRaw(const sourcetype* SrcData, bool bIgnoreAlpha)
 		WriteFrameBufferChannel<OutputFormat>(ImfFrameBuffer, GetRawChannelName(Channel), SrcData + Channel, ChannelOutputBuffers[Channel]);
 	}
 
-	Imf::OutputFile ImfFile(MemFile, Header);
-	ImfFile.setFrameBuffer(ImfFrameBuffer);
+	uint32 FileLength;
+	{
+		// This scope ensures that IMF::Outputfile creates a complete file by closing the file when it goes out of scope.
+		// To complete the file, EXR seeks back into the file and writes the scanline offsets when the file is closed, 
+		// which moves the tellp location. So file length is stored in advance for later use.
+
+		Imf::OutputFile ImfFile(MemFile, Header);
+		ImfFile.setFrameBuffer(ImfFrameBuffer);
 	
-	MemFile.Data.AddUninitialized(Width * Height * NumWriteComponents * (OutputFormat == 2 ? 4 : 2));
-	ImfFile.writePixels(Height);
+		MemFile.Data.AddUninitialized(Width * Height * NumWriteComponents * (OutputFormat == 2 ? 4 : 2));
+		ImfFile.writePixels(Height);
+		FileLength = MemFile.tellp();
+	}
 
-	CompressedData.AddUninitialized(MemFile.tellp());
-	FMemory::Memcpy(CompressedData.GetData(), MemFile.Data.GetData(), MemFile.tellp());
+	CompressedData.AddUninitialized(FileLength);
+	FMemory::Memcpy(CompressedData.GetData(), MemFile.Data.GetData(), FileLength);
 
 	const double DeltaTime = FPlatformTime::Seconds() - StartTime;
 	UE_LOG(LogImageWrapper, Verbose, TEXT("Compressed image in %.3f seconds"), DeltaTime);
