Fix for nounity error

diff --git a/Engine/Source/Runtime/ImageWrapper/Private/TgaImageSupport.cpp b/Engine/Source/Runtime/ImageWrapper/Private/TgaImageSupport.cpp
--- a/Engine/Source/Runtime/ImageWrapper/Private/TgaImageSupport.cpp
+++ b/Engine/Source/Runtime/ImageWrapper/Private/TgaImageSupport.cpp
@@ -1,6 +1,7 @@
 // Copyright Epic Games, Inc. All Rights Reserved.
 
 #include "TgaImageSupport.h"
+#include "ImageWrapperPrivate.h"
 
 namespace TgaImageSupportImpl
 {
