Fix include path casing

diff --git a/Engine/Plugins/Media/ImgMedia/Source/OpenExrWrapper/Private/OpenExrWrapper.cpp b/Engine/Plugins/Media/ImgMedia/Source/OpenExrWrapper/Private/OpenExrWrapper.cpp
--- a/Engine/Plugins/Media/ImgMedia/Source/OpenExrWrapper/Private/OpenExrWrapper.cpp
+++ b/Engine/Plugins/Media/ImgMedia/Source/OpenExrWrapper/Private/OpenExrWrapper.cpp
@@ -7,11 +7,11 @@
 
 PRAGMA_DEFAULT_VISIBILITY_START
 THIRD_PARTY_INCLUDES_START
-	#include "OpenEXR/ImathBox.h"
-	#include "OpenEXR/ImfHeader.h"
-	#include "OpenEXR/ImfRgbaFile.h"
-	#include "OpenEXR/ImfCompressionAttribute.h"
-	#include "OpenEXR/ImfStandardAttributes.h"
+	#include "openexr/ImathBox.h"
+	#include "openexr/ImfHeader.h"
+	#include "openexr/ImfRgbaFile.h"
+	#include "openexr/ImfCompressionAttribute.h"
+	#include "openexr/ImfStandardAttributes.h"
 THIRD_PARTY_INCLUDES_END
 PRAGMA_DEFAULT_VISIBILITY_END
 
diff --git a/Engine/Plugins/MovieScene/MovieRenderPipeline/Source/openexrRTTI/Private/OpenExrRTTIModule.cpp b/Engine/Plugins/MovieScene/MovieRenderPipeline/Source/openexrRTTI/Private/OpenExrRTTIModule.cpp
--- a/Engine/Plugins/MovieScene/MovieRenderPipeline/Source/openexrRTTI/Private/OpenExrRTTIModule.cpp
+++ b/Engine/Plugins/MovieScene/MovieRenderPipeline/Source/openexrRTTI/Private/OpenExrRTTIModule.cpp
@@ -4,9 +4,9 @@
 #include "Modules/ModuleManager.h"
 
 THIRD_PARTY_INCLUDES_START
-#include "OpenEXR/ImfStringAttribute.h"
-#include "OpenEXR/ImfIntAttribute.h"
-#include "OpenEXR/ImfFloatAttribute.h"
+#include "openexr/ImfStringAttribute.h"
+#include "openexr/ImfIntAttribute.h"
+#include "openexr/ImfFloatAttribute.h"
 THIRD_PARTY_INCLUDES_END
 
 
diff --git a/Engine/Plugins/MovieScene/MovieRenderPipeline/Source/openexrRTTI/Public/IOpenExrRTTIModule.h b/Engine/Plugins/MovieScene/MovieRenderPipeline/Source/openexrRTTI/Public/IOpenExrRTTIModule.h
--- a/Engine/Plugins/MovieScene/MovieRenderPipeline/Source/openexrRTTI/Public/IOpenExrRTTIModule.h
+++ b/Engine/Plugins/MovieScene/MovieRenderPipeline/Source/openexrRTTI/Public/IOpenExrRTTIModule.h
@@ -7,7 +7,7 @@
 #include "Containers/UnrealString.h"
 
 THIRD_PARTY_INCLUDES_START
-#include "OpenEXR/ImfHeader.h"
+#include "openexr/ImfHeader.h"
 THIRD_PARTY_INCLUDES_END
 
 
diff --git a/Engine/Source/Runtime/ImageWrapper/Private/Formats/ExrImageWrapper.h b/Engine/Source/Runtime/ImageWrapper/Private/Formats/ExrImageWrapper.h
--- a/Engine/Source/Runtime/ImageWrapper/Private/Formats/ExrImageWrapper.h
+++ b/Engine/Source/Runtime/ImageWrapper/Private/Formats/ExrImageWrapper.h
@@ -9,16 +9,16 @@
 #if WITH_UNREALEXR
 
 THIRD_PARTY_INCLUDES_START
-	#include "OpenEXR/ImfIO.h"
-	#include "OpenEXR/ImathBox.h"
-	#include "OpenEXR/ImfChannelList.h"
-	#include "OpenEXR/ImfInputFile.h"
-	#include "OpenEXR/ImfOutputFile.h"
-	#include "OpenEXR/ImfArray.h"
-	#include "OpenEXR/ImfHeader.h"
-	#include "OpenEXR/ImfStdIO.h"
-	#include "OpenEXR/ImfChannelList.h"
-	#include "OpenEXR/ImfRgbaFile.h"
+	#include "openexr/ImfIO.h"
+	#include "openexr/ImathBox.h"
+	#include "openexr/ImfChannelList.h"
+	#include "openexr/ImfInputFile.h"
+	#include "openexr/ImfOutputFile.h"
+	#include "openexr/ImfArray.h"
+	#include "openexr/ImfHeader.h"
+	#include "openexr/ImfStdIO.h"
+	#include "openexr/ImfChannelList.h"
+	#include "openexr/ImfRgbaFile.h"
 THIRD_PARTY_INCLUDES_END
 
 
