Mirgrated the fix for VLN-1183 into UE5 EA.

diff --git a/Engine/Source/Editor/UnrealEd/Private/Factories/EditorFactories.cpp b/Engine/Source/Editor/UnrealEd/Private/Factories/EditorFactories.cpp
--- a/Engine/Source/Editor/UnrealEd/Private/Factories/EditorFactories.cpp
+++ b/Engine/Source/Editor/UnrealEd/Private/Factories/EditorFactories.cpp
@@ -4502,7 +4502,17 @@ bool UTextureFactory::IsImportResolutionValid(int32 Width, int32 Height, bool bA
 	// Check if the texture is above the supported resolution and prompt the user if they wish to continue if it is
 	if ( Width > MaximumSupportedResolution || Height > MaximumSupportedResolution )
 	{
-		if ((Width * Height) > FMath::Square(MaximumSupportedVirtualTextureResolution))
+		if (Width * Height < 0)
+		{
+			Warn->Log(ELogVerbosity::Error, *FText::Format(
+				NSLOCTEXT("UnrealEd", "Warning_TextureSizeTooLargeOrInvalid", "Texture is too large to import or it has an invalid resolution. The current maximun is {0} pixels"),
+				FText::AsNumber(FMath::Square(MaximumSupportedVirtualTextureResolution))
+				).ToString());
+
+			bValid = false;
+		}
+
+		if ( bValid && (Width * Height) > FMath::Square(MaximumSupportedVirtualTextureResolution))
 		{
 			Warn->Log(ELogVerbosity::Error, *FText::Format(
 				NSLOCTEXT("UnrealEd", "Warning_TextureSizeTooLarge", "Texture is too large to import. The current maximun is {0} pixels"),
diff --git a/Engine/Source/Runtime/ImageWrapper/Private/Formats/TgaImageWrapper.cpp b/Engine/Source/Runtime/ImageWrapper/Private/Formats/TgaImageWrapper.cpp
--- a/Engine/Source/Runtime/ImageWrapper/Private/Formats/TgaImageWrapper.cpp
+++ b/Engine/Source/Runtime/ImageWrapper/Private/Formats/TgaImageWrapper.cpp
@@ -27,7 +27,10 @@ void FTgaImageWrapper::Uncompress(const ERGBFormat InFormat, const int32 InBitDe
 	uint32* TextureData = reinterpret_cast< uint32* >( RawData.GetData() );
 
 	FTGAFileHeader* TgaHeader = reinterpret_cast< FTGAFileHeader* >( CompressedData.GetData() );
-	DecompressTGA_helper( TgaHeader, TextureData, TextureDataSize );
+	if ( !DecompressTGA_helper( TgaHeader, CompressedData.Num(), TextureData, TextureDataSize ) )
+	{
+		SetError(TEXT("Error while decompressing a TGA"));
+	}
 }
 
 bool FTgaImageWrapper::LoadTGAHeader()
diff --git a/Engine/Source/Runtime/ImageWrapper/Private/TgaImageSupport.cpp b/Engine/Source/Runtime/ImageWrapper/Private/TgaImageSupport.cpp
--- a/Engine/Source/Runtime/ImageWrapper/Private/TgaImageSupport.cpp
+++ b/Engine/Source/Runtime/ImageWrapper/Private/TgaImageSupport.cpp
@@ -5,26 +5,33 @@
 
 namespace TgaImageSupportImpl
 {
-	void DecompressTGA_RLE_32bpp( const FTGAFileHeader* TgaHeader, uint32* TextureData )
+	bool DecompressTGA_RLE_32bpp( const FTGAFileHeader* TGA, const int64 TGABufferLenght, uint32* TextureData )
 	{
-		uint8*  IdData      = (uint8*)TgaHeader + sizeof(FTGAFileHeader); 
-		uint8*  ColorMap    = IdData + TgaHeader->IdFieldLength;
-		uint8*  ImageData   = (uint8*) (ColorMap + (TgaHeader->ColorMapEntrySize + 4) / 8 * TgaHeader->ColorMapLength);				
-		uint32  Pixel       = 0;
-		int32   RLERun      = 0;
-		int32   RAWRun      = 0;
-
-		for ( int32 Y = TgaHeader->Height - 1 ; Y >=0; Y-- ) // Y-flipped.
-		{					
-			for ( int32 X = 0; X < TgaHeader->Width; X++ )
-			{						
+		const uint8* const IdData = (uint8*)TGA + sizeof(FTGAFileHeader); 
+		const uint8* const  ColorMap = IdData + TGA->IdFieldLength;
+		const uint8* ImageData = (uint8*) (ColorMap + (TGA->ColorMapEntrySize + 4) / 8 * TGA->ColorMapLength);
+
+		const uint8* TGAEnd = (uint8*)TGA + TGABufferLenght;
+
+		uint32 Pixel = 0;
+		int32 RLERun = 0;
+		int32 RAWRun = 0;
+
+		for(int32 Y = TGA->Height-1; Y >=0; Y--) // Y-flipped.
+		{
+			for(int32 X = 0;X < TGA->Width;X++)
+			{
 				if( RLERun > 0 )
 				{
 					RLERun--;  // reuse current Pixel data.
 				}
 				else if( RAWRun == 0 ) // new raw pixel or RLE-run.
 				{
-					uint8 RLEChunk = *(ImageData++);
+					if ( TGAEnd < ImageData )
+					{
+						return false;
+					}
+					uint8 RLEChunk = *(ImageData++);							
 					if( RLEChunk & 0x80 )
 					{
 						RLERun = ( RLEChunk & 0x7F ) + 1;
@@ -34,7 +41,12 @@ namespace TgaImageSupportImpl
 					{
 						RAWRun = ( RLEChunk & 0x7F ) + 1;
 					}
-				}							
+
+					if ( TGAEnd < ImageData + RAWRun * 4 )
+					{
+						return false;
+					}
+				}
 				// Retrieve new pixel data - raw run or single pixel for RLE stretch.
 				if( RAWRun > 0 )
 				{
@@ -44,28 +56,39 @@ namespace TgaImageSupportImpl
 					RLERun--;
 				}
 				// Store.
-				*( (TextureData + Y*TgaHeader->Width)+X ) = Pixel;
+				*( (TextureData + Y*TGA->Width)+X ) = Pixel;
 			}
 		}
+
+		return true;
 	}
 
-	void DecompressTGA_RLE_24bpp( const FTGAFileHeader* TgaHeader, uint32* TextureData )
+	bool DecompressTGA_RLE_24bpp( const FTGAFileHeader* TGA, const int64 TGABufferLenght, uint32* TextureData )
 	{
-		uint8*  IdData = (uint8*)TgaHeader + sizeof(FTGAFileHeader); 
-		uint8*  ColorMap = IdData + TgaHeader->IdFieldLength;
-		uint8*  ImageData = (uint8*) (ColorMap + (TgaHeader->ColorMapEntrySize + 4) / 8 * TgaHeader->ColorMapLength);
-		uint8   Pixel[4] = {};
-		int32   RLERun = 0;
-		int32   RAWRun = 0;
-
-		for(int32 Y = TgaHeader->Height-1; Y >=0; Y--) // Y-flipped.
-		{					
-			for(int32 X = 0;X < TgaHeader->Width;X++)
-			{						
+		const uint8* const IdData = (uint8*)TGA + sizeof(FTGAFileHeader); 
+		const uint8* const ColorMap = IdData + TGA->IdFieldLength;
+		const uint8* ImageData = (uint8*) (ColorMap + (TGA->ColorMapEntrySize + 4) / 8 * TGA->ColorMapLength);
+
+		const uint8* TGAEnd = (uint8*)TGA + TGABufferLenght;
+
+		uint8 Pixel[4] = {};
+		int32 RLERun = 0;
+		int32 RAWRun = 0;
+
+		for(int32 Y = TGA->Height-1; Y >=0; Y--) // Y-flipped.
+		{
+			for(int32 X = 0;X < TGA->Width;X++)
+			{
 				if( RLERun > 0 )
-					RLERun--;  // reuse current Pixel data.
+				{
+					RLERun--; // reuse current Pixel data.
+				}
 				else if( RAWRun == 0 ) // new raw pixel or RLE-run.
 				{
+					if ( TGAEnd < ImageData )
+					{
+						return false;
+					}
 					uint8 RLEChunk = *(ImageData++);
 					if( RLEChunk & 0x80 )
 					{
@@ -76,7 +99,12 @@ namespace TgaImageSupportImpl
 					{
 						RAWRun = ( RLEChunk & 0x7F ) + 1;
 					}
-				}							
+
+					if ( TGAEnd < ImageData + RAWRun * 3 )
+					{
+						return false;
+					}
+				}
 				// Retrieve new pixel data - raw run or single pixel for RLE stretch.
 				if( RAWRun > 0 )
 				{
@@ -88,31 +116,38 @@ namespace TgaImageSupportImpl
 					RLERun--;
 				}
 				// Store.
-				*( (TextureData + Y*TgaHeader->Width)+X ) = *(uint32*)&Pixel;
+				*( (TextureData + Y*TGA->Width)+X ) = *(uint32*)&Pixel;
 			}
 		}
+		return true;
 	}
 
-	void DecompressTGA_RLE_16bpp( const FTGAFileHeader* TgaHeader, uint32* TextureData )
+	bool DecompressTGA_RLE_16bpp( const FTGAFileHeader* TGA, const int64 TGABufferLenght, uint32* TextureData )
 	{
-		uint8*  IdData = (uint8*)TgaHeader + sizeof(FTGAFileHeader);
-		uint8*  ColorMap = IdData + TgaHeader->IdFieldLength;				
-		uint16* ImageData = (uint16*) (ColorMap + (TgaHeader->ColorMapEntrySize + 4) / 8 * TgaHeader->ColorMapLength);
-		uint16  FilePixel = 0;
-		uint32  TexturePixel = 0;
-		int32   RLERun = 0;
-		int32   RAWRun = 0;
-
-		for ( int32 Y = TgaHeader->Height-1; Y >=0; Y-- ) // Y-flipped.
+		const uint8* const IdData = (uint8*)TGA + sizeof(FTGAFileHeader);
+		const uint8* const ColorMap = IdData + TGA->IdFieldLength;				
+		const uint16* ImageData = (uint16*) (ColorMap + (TGA->ColorMapEntrySize + 4) / 8 * TGA->ColorMapLength);
+
+		const uint16* TGAEnd = (uint16*)TGA + TGABufferLenght / 2;
+
+		uint32 TexturePixel = 0;
+		int32 RLERun = 0;
+		int32 RAWRun = 0;
+
+		for(int32 Y = TGA->Height-1; Y >=0; Y--) // Y-flipped.
 		{					
-			for ( int32 X=0;X<TgaHeader->Width;X++ )
+			for( int32 X=0;X<TGA->Width;X++ )
 			{						
-				if ( RLERun > 0 )
+				if( RLERun > 0 )
 				{
 					RLERun--;  // reuse current Pixel data.
 				}
-				else if ( RAWRun == 0 ) // new raw pixel or RLE-run.
+				else if( RAWRun == 0 ) // new raw pixel or RLE-run.
 				{
+					if ( TGAEnd < ImageData )
+					{
+						return false;
+					}
 					uint8 RLEChunk =  *((uint8*)ImageData);
 					ImageData = (uint16*)(((uint8*)ImageData)+1);
 					if( RLEChunk & 0x80 )
@@ -124,51 +159,95 @@ namespace TgaImageSupportImpl
 					{
 						RAWRun = ( RLEChunk & 0x7F ) + 1;
 					}
-				}
 
+					if ( TGAEnd < ImageData + RAWRun )
+					{
+						return false;
+					}
+				}
 				// Retrieve new pixel data - raw run or single pixel for RLE stretch.
-				if ( RAWRun > 0 )
+				if( RAWRun > 0 )
 				{ 
-					FilePixel = *(ImageData++);
+					const uint16 FilePixel = *(ImageData++);
 					RAWRun--;
 					RLERun--;
-				}
-
-				// Convert file format A1R5G5B5 into pixel format B8G8R8B8
-				TexturePixel = (FilePixel & 0x001F) << 3;
-				TexturePixel |= (FilePixel & 0x03E0) << 6;
-				TexturePixel |= (FilePixel & 0x7C00) << 9;
-				TexturePixel |= (FilePixel & 0x8000) << 16;
 
+					// Convert file format A1R5G5B5 into pixel format B8G8R8B8
+					TexturePixel = (FilePixel & 0x001F) << 3;
+					TexturePixel |= (FilePixel & 0x03E0) << 6;
+					TexturePixel |= (FilePixel & 0x7C00) << 9;
+					TexturePixel |= (FilePixel & 0x8000) << 16;
+				}
 				// Store.
-				*( (TextureData + Y*TgaHeader->Width)+X ) = TexturePixel;
+				*( (TextureData + Y*TGA->Width)+X ) = TexturePixel;
 			}
 		}
+
+		return true;
+	}
+
+	bool DecompressTGA_32bpp( const FTGAFileHeader* TGA, const int64 TGABufferLenght, uint32* TextureData )
+	{
+		const uint8* const IdData = (uint8*)TGA + sizeof(FTGAFileHeader);
+		const uint8* const ColorMap = IdData + TGA->IdFieldLength;
+		const uint32* const ImageData = (uint32*) (ColorMap + (TGA->ColorMapEntrySize + 4) / 8 * TGA->ColorMapLength);
+
+		if ( (uint8*)TGA + TGABufferLenght < (uint8*)(ImageData + TGA->Width * TGA->Height) )
+		{
+			return false;
+		}
+
+		for(int32 Y = 0;Y < TGA->Height;Y++)
+		{
+			FMemory::Memcpy(TextureData + Y * TGA->Width,ImageData + (TGA->Height - Y - 1) * TGA->Width,TGA->Width * 4);
+		}
+
+		return true;
 	}
 
-	void DecompressTGA_32bpp( const FTGAFileHeader* TgaHeader, uint32* TextureData )
+	bool DecompressTGA_24bpp( const FTGAFileHeader* TGA, const int64 TGABufferLenght, uint32* TextureData )
 	{
-		uint8*	IdData = (uint8*)TgaHeader + sizeof(FTGAFileHeader);
-		uint8*	ColorMap = IdData + TgaHeader->IdFieldLength;
-		uint32*	ImageData = (uint32*) (ColorMap + (TgaHeader->ColorMapEntrySize + 4) / 8 * TgaHeader->ColorMapLength);
+		const uint8* const IdData = (uint8*)TGA + sizeof(FTGAFileHeader);
+		const uint8* const ColorMap = IdData + TGA->IdFieldLength;
+		const uint8* const ImageData = (uint8*) (ColorMap + (TGA->ColorMapEntrySize + 4) / 8 * TGA->ColorMapLength);
+		uint8 Pixel[4];
+
+		if ( (uint8*)TGA + TGABufferLenght < ImageData + TGA->Width * TGA->Height * 3 )
+		{
+			return false;
+		}
 
-		for ( int32 Y = 0; Y < TgaHeader->Height; Y++ )
+		for(int32 Y = 0; Y < TGA->Height; Y++)
 		{
-			FMemory::Memcpy(TextureData + Y * TgaHeader->Width,ImageData + (TgaHeader->Height - Y - 1) * TgaHeader->Width,TgaHeader->Width * 4);
+			for(int32 X = 0; X < TGA->Width; X++)
+			{
+				Pixel[0] = *(( ImageData+( TGA->Height-Y-1 )*TGA->Width*3 )+X*3+0);
+				Pixel[1] = *(( ImageData+( TGA->Height-Y-1 )*TGA->Width*3 )+X*3+1);
+				Pixel[2] = *(( ImageData+( TGA->Height-Y-1 )*TGA->Width*3 )+X*3+2);
+				Pixel[3] = 255;
+				*((TextureData+Y*TGA->Width)+X) = *(uint32*)&Pixel;
+			}
 		}
+
+		return true;
 	}
 
-	void DecompressTGA_16bpp( const FTGAFileHeader* TgaHeader, uint32* TextureData )
+	bool DecompressTGA_16bpp( const FTGAFileHeader* TGA, const int64 TGABufferLenght, uint32* TextureData )
 	{
-		uint8*	IdData = (uint8*)TgaHeader + sizeof(FTGAFileHeader);
-		uint8*	ColorMap = IdData + TgaHeader->IdFieldLength;
-		uint16*	ImageData = (uint16*) (ColorMap + (TgaHeader->ColorMapEntrySize + 4) / 8 * TgaHeader->ColorMapLength);
-		uint16    FilePixel = 0;
-		uint32	TexturePixel = 0;
+		const uint8* const IdData = (uint8*)TGA + sizeof(FTGAFileHeader);
+		const uint8* const ColorMap = IdData + TGA->IdFieldLength;
+		const uint16* ImageData = (uint16*) (ColorMap + (TGA->ColorMapEntrySize + 4) / 8 * TGA->ColorMapLength);
+		uint16 FilePixel = 0;
+		uint32 TexturePixel = 0;
 
-		for ( int32 Y = TgaHeader->Height - 1; Y >= 0; Y-- )
+		if ( (uint16*)((uint8*)TGA + TGABufferLenght) < ImageData + TGA->Height * TGA->Width )
+		{
+			return false;
+		}
+
+		for (int32 Y = TGA->Height - 1; Y >= 0; Y--)
 		{					
-			for ( int32 X = 0; X < TgaHeader->Width; X++ )
+			for (int32 X = 0; X < TGA->Width; X++)
 			{
 				FilePixel = *ImageData++;
 				// Convert file format A1R5G5B5 into pixel format B8G8R8A8
@@ -177,64 +256,54 @@ namespace TgaImageSupportImpl
 				TexturePixel |= (FilePixel & 0x7C00) << 9;
 				TexturePixel |= (FilePixel & 0x8000) << 16;
 				// Store.
-				*((TextureData + Y*TgaHeader->Width) + X) = TexturePixel;						
+				*((TextureData + Y*TGA->Width) + X) = TexturePixel;						
 			}
 		}
+
+		return true;
 	}
 
-	void DecompressTGA_24bpp( const FTGAFileHeader* TgaHeader, uint32* TextureData )
+	bool DecompressTGA_8bpp( const FTGAFileHeader* TGA, const int64 TGABufferLenght, uint8* TextureData )
 	{
-		uint8*	IdData = (uint8*)TgaHeader + sizeof(FTGAFileHeader);
-		uint8*	ColorMap = IdData + TgaHeader->IdFieldLength;
-		uint8*	ImageData = (uint8*) (ColorMap + (TgaHeader->ColorMapEntrySize + 4) / 8 * TgaHeader->ColorMapLength);
-		uint8    Pixel[4];
+		const uint8* const IdData = (uint8*)TGA + sizeof(FTGAFileHeader);
+		const uint8* const  ColorMap = IdData + TGA->IdFieldLength;
+		const uint8* const  ImageData = (uint8*) (ColorMap + (TGA->ColorMapEntrySize + 4) / 8 * TGA->ColorMapLength);
 
-		for( int32 Y = 0; Y < TgaHeader->Height; Y++ )
+		if ( (uint8*)TGA + TGABufferLenght < ImageData + TGA->Width * TGA->Height )
 		{
-			for( int32 X = 0; X < TgaHeader->Width; X++ )
-			{
-				Pixel[0] = *(( ImageData+( TgaHeader->Height-Y-1 )*TgaHeader->Width*3 )+X*3+0);
-				Pixel[1] = *(( ImageData+( TgaHeader->Height-Y-1 )*TgaHeader->Width*3 )+X*3+1);
-				Pixel[2] = *(( ImageData+( TgaHeader->Height-Y-1 )*TgaHeader->Width*3 )+X*3+2);
-				Pixel[3] = 255;
-				*((TextureData+Y*TgaHeader->Width)+X) = *(uint32*)&Pixel;
-			}
+			return false;
 		}
-	}
-
-	void DecompressTGA_8bpp( const FTGAFileHeader* TgaHeader, uint8* TextureData )
-	{
-		const uint8*  const IdData = (uint8*)TgaHeader + sizeof(FTGAFileHeader);
-		const uint8*  const ColorMap = IdData + TgaHeader->IdFieldLength;
-		const uint8*  const ImageData = (uint8*) (ColorMap + (TgaHeader->ColorMapEntrySize + 4) / 8 * TgaHeader->ColorMapLength);
 
 		int32 RevY = 0;
-		for (int32 Y = TgaHeader->Height-1; Y >= 0; --Y)
+		for (int32 Y = TGA->Height-1; Y >= 0; --Y)
 		{
-			const uint8* ImageCol = ImageData + (Y * TgaHeader->Width); 
-			uint8* TextureCol = TextureData + (RevY++ * TgaHeader->Width);
-			FMemory::Memcpy(TextureCol, ImageCol, TgaHeader->Width);
+			const uint8* ImageCol = ImageData + (Y * TGA->Width); 
+			uint8* TextureCol = TextureData + (RevY++ * TGA->Width);
+			FMemory::Memcpy(TextureCol, ImageCol, TGA->Width);
 		}
+
+		return true;
 	}
 }
 
-bool DecompressTGA_helper( const FTGAFileHeader* TgaHeader, uint32*& TextureData, const int32 TextureDataSize )
+bool DecompressTGA_helper( const FTGAFileHeader* TgaHeader, const int64 TGABufferLenght, uint32*& TextureData, const int32 TextureDataSize )
 {
+	bool bSuccess = false;
 	if ( TgaHeader->ImageTypeCode == 10 ) // 10 = RLE compressed 
 	{
 		// RLE compression: CHUNKS: 1 -byte header, high bit 0 = raw, 1 = compressed
 		// bits 0-6 are a 7-bit count; count+1 = number of raw pixels following, or rle pixels to be expanded. 
 		if(TgaHeader->BitsPerPixel == 32)
 		{
-			TgaImageSupportImpl::DecompressTGA_RLE_32bpp(TgaHeader, TextureData);
+			bSuccess = TgaImageSupportImpl::DecompressTGA_RLE_32bpp(TgaHeader, TGABufferLenght, TextureData);
 		}
 		else if( TgaHeader->BitsPerPixel == 24 )
-		{	
-			TgaImageSupportImpl::DecompressTGA_RLE_24bpp(TgaHeader, TextureData);
+		{
+			bSuccess = TgaImageSupportImpl::DecompressTGA_RLE_24bpp(TgaHeader, TGABufferLenght, TextureData);
 		}
 		else if( TgaHeader->BitsPerPixel == 16 )
 		{
-			TgaImageSupportImpl::DecompressTGA_RLE_16bpp(TgaHeader, TextureData);
+			bSuccess = TgaImageSupportImpl::DecompressTGA_RLE_16bpp(TgaHeader, TGABufferLenght, TextureData);
 		}
 		else
 		{
@@ -246,15 +315,15 @@ bool DecompressTGA_helper( const FTGAFileHeader* TgaHeader, uint32*& TextureData
 	{
 		if(TgaHeader->BitsPerPixel == 32)
 		{
-			TgaImageSupportImpl::DecompressTGA_32bpp(TgaHeader, TextureData);
+			bSuccess = TgaImageSupportImpl::DecompressTGA_32bpp(TgaHeader, TGABufferLenght, TextureData);
 		}
 		else if(TgaHeader->BitsPerPixel == 16)
 		{
-			TgaImageSupportImpl::DecompressTGA_16bpp(TgaHeader, TextureData);
+			bSuccess = TgaImageSupportImpl::DecompressTGA_16bpp(TgaHeader, TGABufferLenght, TextureData);
 		}
 		else if(TgaHeader->BitsPerPixel == 24)
 		{
-			TgaImageSupportImpl::DecompressTGA_24bpp(TgaHeader, TextureData);
+			bSuccess = TgaImageSupportImpl::DecompressTGA_24bpp(TgaHeader, TGABufferLenght, TextureData);
 		}
 		else
 		{
@@ -265,12 +334,12 @@ bool DecompressTGA_helper( const FTGAFileHeader* TgaHeader, uint32*& TextureData
 	// Support for alpha stored as pseudo-color 8-bit TgaHeader
 	else if(TgaHeader->ColorMapType == 1 && TgaHeader->ImageTypeCode == 1 && TgaHeader->BitsPerPixel == 8)
 	{
-		TgaImageSupportImpl::DecompressTGA_8bpp(TgaHeader, (uint8*)TextureData);
+		bSuccess = TgaImageSupportImpl::DecompressTGA_8bpp(TgaHeader, TGABufferLenght, (uint8*)TextureData);
 	}
 	// standard grayscale
 	else if(TgaHeader->ColorMapType == 0 && TgaHeader->ImageTypeCode == 3 && TgaHeader->BitsPerPixel == 8)
 	{
-		TgaImageSupportImpl::DecompressTGA_8bpp(TgaHeader, (uint8*)TextureData);
+		bSuccess = TgaImageSupportImpl::DecompressTGA_8bpp(TgaHeader, TGABufferLenght, (uint8*)TextureData);
 	}
 	else
 	{
@@ -278,6 +347,12 @@ bool DecompressTGA_helper( const FTGAFileHeader* TgaHeader, uint32*& TextureData
 		return false;
 	}
 
+	if (!bSuccess)
+	{
+		UE_LOG(LogImageWrapper, Error, TEXT("The TGA file is invalid or corrupted"));
+		return false;
+	}
+
 	// Flip the image data if the flip bits are set in the TgaHeader header.
 	const bool bFlipX = (TgaHeader->ImageDescriptor & 0x10) ? 1 : 0;
 	const bool bFlipY = (TgaHeader->ImageDescriptor & 0x20) ? 1 : 0;
diff --git a/Engine/Source/Runtime/ImageWrapper/Public/TgaImageSupport.h b/Engine/Source/Runtime/ImageWrapper/Public/TgaImageSupport.h
--- a/Engine/Source/Runtime/ImageWrapper/Public/TgaImageSupport.h
+++ b/Engine/Source/Runtime/ImageWrapper/Public/TgaImageSupport.h
@@ -36,4 +36,4 @@ struct FTGAFileHeader
  * This helper allows to decompress TGA data in a pre-allocated memory block.
  * The pixel format is necessarily PF_A8R8G8B8.
  */
-IMAGEWRAPPER_API bool DecompressTGA_helper(	const FTGAFileHeader* TGA, uint32*& TextureData, const int32 TextureDataSize );
+IMAGEWRAPPER_API bool DecompressTGA_helper(	const FTGAFileHeader* TGA, const int64 TGABufferLenght, uint32*& TextureData, const int32 TextureDataSize );
