Avoid having the ImageWrapper module depend on CoreUObject so that it can be used from standalone texture build workers that don't want to link to CoreUObject.

diff --git a/Engine/Source/Runtime/ImageWrapper/ImageWrapper.Build.cs b/Engine/Source/Runtime/ImageWrapper/ImageWrapper.Build.cs
--- a/Engine/Source/Runtime/ImageWrapper/ImageWrapper.Build.cs
+++ b/Engine/Source/Runtime/ImageWrapper/ImageWrapper.Build.cs
@@ -16,7 +16,6 @@ public class ImageWrapper : ModuleRules
 		PublicDefinitions.Add("WITH_UNREALJPEG=1");
 
 		PrivateDependencyModuleNames.Add("Core");
-		PublicDependencyModuleNames.Add("CoreUObject");
 
 		AddEngineThirdPartyPrivateStaticDependencies(Target,
 			"zlib",
diff --git a/Engine/Source/Runtime/ImageWrapper/Public/BmpImageSupport.h b/Engine/Source/Runtime/ImageWrapper/Public/BmpImageSupport.h
--- a/Engine/Source/Runtime/ImageWrapper/Public/BmpImageSupport.h
+++ b/Engine/Source/Runtime/ImageWrapper/Public/BmpImageSupport.h
@@ -2,11 +2,8 @@
 
 #pragma once
 
-#include "UObject/ObjectMacros.h"
 #include "CoreMinimal.h"
 
-#include "BmpImageSupport.generated.h"
-
 struct FBitmapFileHeader;
 struct FBitmapInfoHeader;
 
@@ -23,7 +20,6 @@ enum EBitmapCompression
 };
 
 // Bitmap info header versions.
-UENUM()
 enum class EBitmapHeaderVersion : uint8
 {
 	BHV_BITMAPINFOHEADER   = 0,
@@ -34,7 +30,6 @@ enum class EBitmapHeaderVersion : uint8
 };
 
 // Color space type of the bitmap, property introduced in Bitmap header version 4.
-UENUM()
 enum class EBitmapCSType : uint32
 {
 	BCST_BLCS_CALIBRATED_RGB     = 0x00000000,
