Moving HDR loader from runtime engine to imagewrapper.
Adding support of hdr to interchange those are now imported as standart 2d Texture.

diff --git a/Engine/Plugins/Runtime/Interchange/Source/Import/Private/InterchangeImportModule.cpp b/Engine/Plugins/Runtime/Interchange/Source/Import/Private/InterchangeImportModule.cpp
--- a/Engine/Plugins/Runtime/Interchange/Source/Import/Private/InterchangeImportModule.cpp
+++ b/Engine/Plugins/Runtime/Interchange/Source/Import/Private/InterchangeImportModule.cpp
@@ -14,6 +14,7 @@
 #include "Texture/InterchangeBMPTranslator.h"
 #include "Texture/InterchangeDDSTranslator.h"
 #include "Texture/InterchangeEXRTranslator.h"
+#include "Texture/InterchangeHDRTranslator.h"
 #include "Texture/InterchangeJPGTranslator.h"
 #include "Texture/InterchangePCXTranslator.h"
 #include "Texture/InterchangePNGTranslator.h"
@@ -59,6 +60,7 @@ void FInterchangeImportModule::StartupModule()
 		InterchangeManager.RegisterTranslator(UInterchangePNGTranslator::StaticClass());
 		InterchangeManager.RegisterTranslator(UInterchangePSDTranslator::StaticClass());
 		InterchangeManager.RegisterTranslator(UInterchangeTGATranslator::StaticClass());
+		InterchangeManager.RegisterTranslator(UInterchangeHDRTranslator::StaticClass());
 
 		//Register the factories
 		InterchangeManager.RegisterFactory(UInterchangeTextureFactory::StaticClass());
diff --git a/Engine/Plugins/Runtime/Interchange/Source/Import/Private/Texture/InterchangeHDRTranslator.cpp b/Engine/Plugins/Runtime/Interchange/Source/Import/Private/Texture/InterchangeHDRTranslator.cpp
new file mode 100644
--- /dev/null
+++ b/Engine/Plugins/Runtime/Interchange/Source/Import/Private/Texture/InterchangeHDRTranslator.cpp
@@ -0,0 +1,108 @@
+// Copyright Epic Games, Inc. All Rights Reserved. 
+
+#include "Texture/InterchangeHDRTranslator.h"
+
+#include "Engine/Texture.h"
+#include "Engine/Texture2D.h"
+#include "Formats/HdrImageWrapper.h"
+#include "InterchangeImportLog.h"
+#include "InterchangeTextureNode.h"
+#include "Misc/ConfigCacheIni.h"
+#include "Misc/FileHelper.h"
+#include "Misc/Paths.h"
+#include "Misc/ScopedSlowTask.h"
+#include "Modules/ModuleManager.h"
+#include "Nodes/InterchangeBaseNodeContainer.h"
+#include "Texture/TextureTranslatorUtilities.h"
+
+
+bool UInterchangeHDRTranslator::CanImportSourceData(const UInterchangeSourceData* InSourceData) const
+{
+	FString Extension = FPaths::GetExtension(InSourceData->GetFilename());
+	FString HDRExtension = (TEXT("hdr"));
+	return HDRExtension.StartsWith(Extension);
+}
+
+bool UInterchangeHDRTranslator::Translate(const UInterchangeSourceData* SourceData, UInterchangeBaseNodeContainer& BaseNodeContainer) const
+{
+	return UE::Interchange::FTextureTranslatorUtilities::Generic2DTextureTranslate(SourceData, BaseNodeContainer);
+}
+
+TOptional<UE::Interchange::FImportImage> UInterchangeHDRTranslator::GetTexturePayloadData(const UInterchangeSourceData* SourceData, const FString& PayLoadKey) const
+{
+	if (!SourceData)
+	{
+		UE_LOG(LogInterchangeImport, Error, TEXT("Failed to import HDR, bad source data."));
+		return TOptional<UE::Interchange::FImportImage>();
+	}
+
+	TArray64<uint8> SourceDataBuffer;
+	FString Filename = SourceData->GetFilename();
+
+	//Make sure the key fit the filename, The key should always be valid
+	if (!Filename.Equals(PayLoadKey))
+	{
+		UE_LOG(LogInterchangeImport, Error, TEXT("Failed to import HDR, wrong payload key. [%s]"), *Filename);
+		return TOptional<UE::Interchange::FImportImage>();
+	}
+
+	if (!FPaths::FileExists(Filename))
+	{
+		UE_LOG(LogInterchangeImport, Error, TEXT("Failed to import HDR, cannot open file. [%s]"), *Filename);
+		return TOptional<UE::Interchange::FImportImage>();
+	}
+
+	if (!FFileHelper::LoadFileToArray(SourceDataBuffer, *Filename))
+	{
+		UE_LOG(LogInterchangeImport, Error, TEXT("Failed to import HDR, cannot load file content into an array. [%s]"), *Filename);
+		return TOptional<UE::Interchange::FImportImage>();
+	}
+
+	const uint8* Buffer = SourceDataBuffer.GetData();
+	const uint8* BufferEnd = Buffer + SourceDataBuffer.Num();
+
+	bool bAllowNonPowerOfTwo = false;
+	GConfig->GetBool(TEXT("TextureImporter"), TEXT("AllowNonPowerOfTwoTextures"), bAllowNonPowerOfTwo, GEditorIni);
+
+
+	const int32 Length = BufferEnd - Buffer;
+
+
+	UE::Interchange::FImportImage PayloadData;
+
+	//
+	// HDR File
+	//	
+	FHdrImageWrapper HdrImageWrapper;
+	if (HdrImageWrapper.SetCompressedFromView(TArrayView64<const uint8>(Buffer, Length)))
+	{
+		if (!UE::Interchange::FImportImageHelper::IsImportResolutionValid(HdrImageWrapper.GetWidth(), HdrImageWrapper.GetHeight(), bAllowNonPowerOfTwo))
+		{
+			UE_LOG(LogInterchangeImport, Error, TEXT("Failed to import HDR, invalid resolution. Resolution[%d, %d], AllowPowerOfTwo[%s], [%s]"), HdrImageWrapper.GetWidth(), HdrImageWrapper.GetHeight(), bAllowNonPowerOfTwo ? TEXT("True") : TEXT("false"), *Filename);
+			return TOptional<UE::Interchange::FImportImage>();
+		}
+
+		PayloadData.Init2DWithParams(
+			HdrImageWrapper.GetWidth(),
+			HdrImageWrapper.GetHeight(),
+			TSF_BGRE8,
+			false);
+
+		if (!HdrImageWrapper.GetRaw(ERGBFormat::BGRE, 8, PayloadData.RawData))
+		{
+			UE_LOG(LogInterchangeImport, Error, TEXT("Failed to Import HDR. [%s]"), *Filename);
+			UE_LOG(LogInterchangeImport, Error, TEXT("%s"), *( HdrImageWrapper.GetErrorMessage().ToString()));
+			return TOptional<UE::Interchange::FImportImage>();
+		}
+
+		PayloadData.CompressionSettings = TC_HDR;
+	}
+	else
+	{
+		UE_LOG(LogInterchangeImport, Error, TEXT("Failed to decode HDR. [%s]"), *Filename);
+		UE_LOG(LogInterchangeImport, Error, TEXT("%s"), *( HdrImageWrapper.GetErrorMessage().ToString()));
+		return TOptional<UE::Interchange::FImportImage>();
+	}
+
+	return PayloadData;
+}
\ No newline at end of file
diff --git a/Engine/Plugins/Runtime/Interchange/Source/Import/Public/Texture/InterchangeHDRTranslator.h b/Engine/Plugins/Runtime/Interchange/Source/Import/Public/Texture/InterchangeHDRTranslator.h
new file mode 100644
--- /dev/null
+++ b/Engine/Plugins/Runtime/Interchange/Source/Import/Public/Texture/InterchangeHDRTranslator.h
@@ -0,0 +1,49 @@
+// Copyright Epic Games, Inc. All Rights Reserved.
+
+#pragma once
+
+#include "CoreMinimal.h"
+#include "UObject/Object.h"
+#include "UObject/ObjectMacros.h"
+#include "InterchangeTranslatorBase.h"
+#include "Nodes/InterchangeBaseNodeContainer.h"
+#include "Texture/InterchangeTexturePayloadData.h"
+#include "Texture/InterchangeTexturePayloadInterface.h"
+
+#include "InterchangeHDRTranslator.generated.h"
+
+UCLASS(BlueprintType)
+class INTERCHANGEIMPORT_API UInterchangeHDRTranslator : public UInterchangeTranslatorBase, public IInterchangeTexturePayloadInterface
+{
+	GENERATED_BODY()
+public:
+
+	/*
+	 * return true if the translator can translate the specified file.
+	 */
+	virtual bool CanImportSourceData(const UInterchangeSourceData* InSourceData) const override;
+
+	/**
+	 * Translate a source data into a node hold by the specified nodes container.
+	 *
+	 * @param SourceData - The source data containing the data to translate
+	 * @param BaseNodeContainer - The unreal objects descriptions container where to put the translated source data.
+	 * @return true if the translator can translate the source data, false otherwise.
+	 */
+	virtual bool Translate(const UInterchangeSourceData* SourceData, UInterchangeBaseNodeContainer& BaseNodeContainer) const override;
+
+
+	/* IInterchangeTexturePayloadInterface Begin */
+
+	/**
+	 * Once the translation is done, the import process need a way to retrieve payload data.
+	 * This payload will be use by the factories to create the asset.
+	 *
+	 * @param SourceData - The source data containing the data to translate
+	 * @param PayloadKey - The key to retrieve the a particular payload contain into the specified source data.
+	 * @return a PayloadData containing the import image data. The TOptional will not be set if there is an error.
+	 */
+	virtual TOptional<UE::Interchange::FImportImage> GetTexturePayloadData(const UInterchangeSourceData* SourceData, const FString& PayLoadKey) const override;
+
+	/* IInterchangeTexturePayloadInterface End */
+};
diff --git a/Engine/Source/Editor/UnrealEd/Private/Factories/EditorFactories.cpp b/Engine/Source/Editor/UnrealEd/Private/Factories/EditorFactories.cpp
--- a/Engine/Source/Editor/UnrealEd/Private/Factories/EditorFactories.cpp
+++ b/Engine/Source/Editor/UnrealEd/Private/Factories/EditorFactories.cpp
@@ -199,7 +199,7 @@
 #include "IAssetTools.h"
 
 #include "DDSLoader.h"
-#include "HDRLoader.h"
+#include "Formats/HdrImageWrapper.h"
 #include "Factories/TIFFLoader.h"
 #include "IESConverter.h"
 #include "IImageWrapper.h"
@@ -3877,30 +3877,37 @@ UTexture* UTextureFactory::ImportTexture(UClass* Class, UObject* InParent, FName
 	//
 	// HDR File
 	//
-	FHDRLoadHelper           HDRLoadHelper(Buffer, Length);
-	if(HDRLoadHelper.IsValid())
+	IImageWrapperModule& ImageWrapperModule = FModuleManager::LoadModuleChecked<IImageWrapperModule>(FName("ImageWrapper"));
+	if (ImageWrapperModule.DetectImageFormat(Buffer, Length) == EImageFormat::HDR)
 	{
-		TArray<uint8> DDSFile;
-		HDRLoadHelper.ExtractDDSInRGBE(DDSFile);
-		FDDSLoadHelper HDRDDSLoadHelper(DDSFile.GetData(), DDSFile.Num());
-
-		// create the cube texture
-		UTextureCube* TextureCube = CreateTextureCube(InParent, Name, Flags);
-		if ( TextureCube )
+		FHdrImageWrapper HdrImageWrapper;
+		if (HdrImageWrapper.SetCompressedFromView(TArrayView64<const uint8>(Buffer, Length)))
 		{
-			TextureCube->Source.Init(
-				HDRDDSLoadHelper.DDSHeader->dwWidth,
-				HDRDDSLoadHelper.DDSHeader->dwHeight,
-				/*NumSlices=*/ 1,
-				/*NumMips=*/ 1,
-				TSF_BGRE8,
-				HDRDDSLoadHelper.GetDDSDataPointer()
-				);
-			// the loader can suggest a compression setting
-			TextureCube->CompressionSettings = TC_HDR;
+			TArray64<uint8> UnCompressedData;
+			if (HdrImageWrapper.GetRaw(ERGBFormat::BGRE, 8, UnCompressedData))
+			{
+				// create the cube texture
+				UTextureCube* TextureCube = CreateTextureCube(InParent, Name, Flags);
+				if ( TextureCube )
+				{
+					TextureCube->Source.Init(
+						HdrImageWrapper.GetWidth(),
+						HdrImageWrapper.GetHeight(),
+						/*NumSlices=*/ 1,
+						/*NumMips=*/ 1,
+						TSF_BGRE8,
+						UnCompressedData.GetData()
+						);
+					// the loader can suggest a compression setting
+					TextureCube->CompressionSettings = TC_HDR;
+
+					return TextureCube;
+				}
+			}
 		}
 
-		return TextureCube;
+		Warn->Log(ELogVerbosity::Error, HdrImageWrapper.GetErrorMessage().ToString());
+		
 	}
 
 	//
diff --git a/Engine/Source/Editor/UnrealEd/UnrealEd.Build.cs b/Engine/Source/Editor/UnrealEd/UnrealEd.Build.cs
--- a/Engine/Source/Editor/UnrealEd/UnrealEd.Build.cs
+++ b/Engine/Source/Editor/UnrealEd/UnrealEd.Build.cs
@@ -37,7 +37,6 @@ public class UnrealEd : ModuleRules
 				"DesktopPlatform",
 				"LauncherPlatform",
 				"GameProjectGeneration",
-				"ImageWrapper",
 				"MainFrame",
 				"TurnkeySupport",
 				"MaterialEditor",
@@ -135,6 +134,7 @@ public class UnrealEd : ModuleRules
 				"EditorStyle",
 				"EngineSettings",
 				"IESFile",
+				"ImageWrapper",
 				"ImageWriteQueue",
 				"InputCore",
 				"InputBindingEditor",
@@ -233,7 +233,6 @@ public class UnrealEd : ModuleRules
 				"SettingsEditor",
 				"StringTableEditor",
 				"FoliageEdit",
-				"ImageWrapper",
 				"Blutility",
 				"IntroTutorials",
 				"WorkspaceMenuStructure",
diff --git a/Engine/Source/Runtime/Engine/Private/HDRLoader.cpp b/Engine/Source/Runtime/Engine/Private/HDRLoader.cpp
--- a/Engine/Source/Runtime/Engine/Private/HDRLoader.cpp
+++ b/Engine/Source/Runtime/Engine/Private/HDRLoader.cpp
@@ -175,6 +175,8 @@ void FHDRLoadHelper::DecompressScanline(uint8* Out, const uint8*& In) const
 
 void FHDRLoadHelper::OldDecompressScanline(uint8* Out, const uint8*& In, uint32 Len)
 {
+	// This is implementation does not have the same behavior as the code found in Radiance but It won't be fix since this implementation of the hdr loader is deprecated.
+
 	int32 Shift = 0;
 
 	while(Len > 0)
diff --git a/Engine/Source/Runtime/Engine/Private/ImageUtils.cpp b/Engine/Source/Runtime/Engine/Private/ImageUtils.cpp
--- a/Engine/Source/Runtime/Engine/Private/ImageUtils.cpp
+++ b/Engine/Source/Runtime/Engine/Private/ImageUtils.cpp
@@ -5,17 +5,17 @@ ImageUtils.cpp: Image utility functions.
 =============================================================================*/
 
 #include "ImageUtils.h"
+
+#include "CubemapUnwrapUtils.h"
+#include "DDSLoader.h"
 #include "Engine/Texture2D.h"
 #include "Engine/TextureCube.h"
-#include "Misc/ObjectThumbnail.h"
 #include "Engine/TextureRenderTarget2D.h"
-#include "CubemapUnwrapUtils.h"
-#include "Logging/MessageLog.h"
 #include "IImageWrapper.h"
 #include "IImageWrapperModule.h"
+#include "Logging/MessageLog.h"
 #include "Misc/FileHelper.h"
-#include "DDSLoader.h"
-#include "HDRLoader.h"
+#include "Misc/ObjectThumbnail.h"
 #include "Misc/Paths.h"
 #include "Modules/ModuleManager.h"
 
@@ -851,7 +851,7 @@ UTexture2D* FImageUtils::ImportFileAsTexture2D(const FString& Filename)
 	IImageWrapperModule& ImageWrapperModule = FModuleManager::Get().LoadModuleChecked<IImageWrapperModule>(TEXT("ImageWrapper"));
 
 	UTexture2D* NewTexture = nullptr;
-	TArray<uint8> Buffer;
+	TArray64<uint8> Buffer;
 	if (FFileHelper::LoadFileToArray(Buffer, *Filename))
 	{
 		EPixelFormat PixelFormat = PF_Unknown;
@@ -863,32 +863,41 @@ UTexture2D* FImageUtils::ImportFileAsTexture2D(const FString& Filename)
 
 		if (FPaths::GetExtension(Filename) == TEXT("HDR"))
 		{
-			FHDRLoadHelper HDRLoadHelper(Buffer.GetData(), Buffer.GetAllocatedSize());
-			if(HDRLoadHelper.IsValid())
+			TSharedPtr<IImageWrapper> HdrImageWrapper =  ImageWrapperModule.CreateImageWrapper(EImageFormat::HDR);
+	
+			if(HdrImageWrapper->SetCompressed(Buffer.GetData(), Buffer.Num()))
 			{
-				TArray<uint8> DDSFile;
-				HDRLoadHelper.ExtractDDSInRGBE(DDSFile);
-				FDDSLoadHelper HDRDDSLoadHelper(DDSFile.GetData(), DDSFile.Num());
+				PixelFormat = PF_FloatRGBA;
+				Width = HdrImageWrapper->GetWidth();
+				Height = HdrImageWrapper->GetHeight();
 
-				if (HDRDDSLoadHelper.IsValid2DTexture())
+				TArray64<uint8> BGREImage;
+				if (HdrImageWrapper->GetRaw(ERGBFormat::BGRE, 8, BGREImage))
 				{
-					PixelFormat = HDRDDSLoadHelper.ComputePixelFormat();
-					Width = HDRDDSLoadHelper.DDSHeader->dwWidth;
-					Height = HDRDDSLoadHelper.DDSHeader->dwHeight;
-
 					NewTexture = UTexture2D::CreateTransient(Width, Height, PixelFormat);
 					if (NewTexture)
 					{
 						uint8* MipData = static_cast<uint8*>(NewTexture->PlatformData->Mips[0].BulkData.Lock(LOCK_READ_WRITE));
 
+						TArrayView64<FColor> SourceColors(reinterpret_cast<FColor*>(BGREImage.GetData()), BGREImage.Num() / sizeof(FColor));
+
 						// Bulk data was already allocated for the correct size when we called CreateTransient above
-						FMemory::Memcpy(MipData, HDRDDSLoadHelper.GetDDSDataPointer(), NewTexture->PlatformData->Mips[0].BulkData.GetBulkDataSize());
+						TArrayView64<FFloat16> Destination(reinterpret_cast<FFloat16*>(MipData), NewTexture->PlatformData->Mips[0].BulkData.GetBulkDataSize() / sizeof(FFloat16));
+
+						int64 DestinationIndex = 0;
+						for (const FColor& Color: SourceColors)
+						{
+							FLinearColor LinearColor = Color.FromRGBE();
+							Destination[DestinationIndex++].Set(LinearColor.R);
+							Destination[DestinationIndex++].Set(LinearColor.G);
+							Destination[DestinationIndex++].Set(LinearColor.B);
+							Destination[DestinationIndex++].Set(LinearColor.A);
+						}
 
 						NewTexture->PlatformData->Mips[0].BulkData.Unlock();
 
 						NewTexture->UpdateResource();
 					}
-
 				}
 			}
 		}
@@ -910,11 +919,11 @@ UTexture2D* FImageUtils::ImportFileAsTexture2D(const FString& Filename)
 	return NewTexture;
 }
 
-UTexture2D* FImageUtils::ImportBufferAsTexture2D(const TArray<uint8>& Buffer)
+UTexture2D* FImageUtils::ImportBufferAsTexture2D(TArrayView64<const uint8> Buffer)
 {
 	IImageWrapperModule& ImageWrapperModule = FModuleManager::Get().LoadModuleChecked<IImageWrapperModule>(TEXT("ImageWrapper"));
 
-	EImageFormat Format = ImageWrapperModule.DetectImageFormat(Buffer.GetData(), Buffer.GetAllocatedSize());
+	EImageFormat Format = ImageWrapperModule.DetectImageFormat(Buffer.GetData(), Buffer.Num());
 
 	UTexture2D* NewTexture = nullptr;
 	EPixelFormat PixelFormat = PF_Unknown;
@@ -927,7 +936,7 @@ UTexture2D* FImageUtils::ImportBufferAsTexture2D(const TArray<uint8>& Buffer)
 		int32 Width = 0;
 		int32 Height = 0;
 
-		if (ImageWrapper->SetCompressed((void*)Buffer.GetData(), Buffer.GetAllocatedSize()))
+		if (ImageWrapper->SetCompressed((void*)Buffer.GetData(), Buffer.Num()))
 		{
 			PixelFormat = PF_Unknown;
 			
@@ -941,7 +950,7 @@ UTexture2D* FImageUtils::ImportBufferAsTexture2D(const TArray<uint8>& Buffer)
 			if (BitDepth == 16)
 			{
 				PixelFormat = PF_FloatRGBA;
-				RGBFormat = ERGBFormat::BGRA;
+				RGBFormat = ERGBFormat::RGBA;
 			}
 			else if (BitDepth == 8)
 			{
@@ -980,6 +989,11 @@ UTexture2D* FImageUtils::ImportBufferAsTexture2D(const TArray<uint8>& Buffer)
 	return NewTexture;
 }
 
+UTexture2D* FImageUtils::ImportBufferAsTexture2D(const TArray<uint8>& Buffer)
+{
+	return ImportBufferAsTexture2D(TArrayView64<const uint8>(Buffer.GetData(), Buffer.Num()));
+}
+
 bool FImageUtils::ExportRenderTargetCubeAsHDR(UTextureRenderTargetCube* TexRT, FArchive& Ar)
 {
 	FHDRExportHelper Exporter;
diff --git a/Engine/Source/Runtime/Engine/Private/VT/VirtualTextureDataBuilder.cpp b/Engine/Source/Runtime/Engine/Private/VT/VirtualTextureDataBuilder.cpp
--- a/Engine/Source/Runtime/Engine/Private/VT/VirtualTextureDataBuilder.cpp
+++ b/Engine/Source/Runtime/Engine/Private/VT/VirtualTextureDataBuilder.cpp
@@ -219,6 +219,7 @@ struct FPixelDataRectangle
 			ImageWrapper->SetRaw(Data, BytesPerPixel * Width * Height, Width, Height, ERGBFormat::BGRA, 8);
 			break;
 		case TSF_BGRE8:
+			// This will probably result in bad image. (Should we convert the data or export it in a format that support float(ERGBFormat::BGRE)?)
 			ImageWrapper->SetRaw(Data, BytesPerPixel * Width * Height, Width, Height, ERGBFormat::BGRA, 8);
 			break;
 		case TSF_RGBA16:
diff --git a/Engine/Source/Runtime/Engine/Public/HDRLoader.h b/Engine/Source/Runtime/Engine/Public/HDRLoader.h
--- a/Engine/Source/Runtime/Engine/Public/HDRLoader.h
+++ b/Engine/Source/Runtime/Engine/Public/HDRLoader.h
@@ -9,6 +9,7 @@
 // http://paulbourke.net/dataformats/pic/
 
 /** To load the HDR file image format. Does not support all possible types HDR formats (e.g. xyze is not supported) */
+class UE_DEPRECATED(5.1, "HDR Loader as been replaced by the FHdrImageWrapper which can be found in the image wrapper module.") FHDRLoadHelper;
 class FHDRLoadHelper
 {
 public:
diff --git a/Engine/Source/Runtime/Engine/Public/ImageUtils.h b/Engine/Source/Runtime/Engine/Public/ImageUtils.h
--- a/Engine/Source/Runtime/Engine/Public/ImageUtils.h
+++ b/Engine/Source/Runtime/Engine/Public/ImageUtils.h
@@ -207,6 +207,8 @@ public:
 	/**
 	 * Imports a texture a buffer and creates Texture2D from it
 	 */
+	ENGINE_API static UTexture2D* ImportBufferAsTexture2D(TArrayView64<const uint8> Buffer);
+
 	ENGINE_API static UTexture2D* ImportBufferAsTexture2D(const TArray<uint8>& Buffer);
 	
 	/**
diff --git a/Engine/Source/Runtime/ImageWrapper/Private/Formats/HdrImageWrapper.cpp b/Engine/Source/Runtime/ImageWrapper/Private/Formats/HdrImageWrapper.cpp
new file mode 100644
--- /dev/null
+++ b/Engine/Source/Runtime/ImageWrapper/Private/Formats/HdrImageWrapper.cpp
@@ -0,0 +1,380 @@
+// Copyright Epic Games, Inc. All Rights Reserved.
+
+#include "Formats/HdrImageWrapper.h"
+
+#define LOCTEXT_NAMESPACE "HdrImageWrapper"
+
+namespace UE::ImageWrapper::Private::HdrImageWrapper
+{
+	FText GetEndOFBufferErrorMessage()
+	{
+		return LOCTEXT("EndOFBufferError","Reached the end of the buffer before finishing decompressing the hdr. The hdr image is likely corrupted");
+	}
+
+	FText GetImageDoneBeforeEndOfBufferErrorMessage()
+	{
+		return LOCTEXT("IMageDoneButThereIsStilSomeDataToDecompress", "Reached the end of the raw image before finishing decompressing the hdr. The hdr image is likely to be corrupted");
+	}
+}
+
+
+bool FHdrImageWrapper::SetCompressedFromView(TArrayView64<const uint8> Data)
+{
+	CompressedData = Data;
+
+	if (CompressedData.Num() < 11)
+	{
+		FreeCompressedData();
+		return false;
+	}
+
+	const uint8* FileDataPtr = CompressedData.GetData();
+
+	char Line[256];
+
+	GetHeaderLine(FileDataPtr, Line);
+
+	if (FCStringAnsi::Strcmp(Line, "#?RADIANCE"))
+	{
+		FreeCompressedData();
+		return false;
+	}
+
+	// Todo add validation for the format
+	bool bHasRGBEFormat = false;
+	while (GetHeaderLine(FileDataPtr, Line) && !bHasRGBEFormat)
+	{
+		if (!FCStringAnsi::Strcmp(Line,"FORMAT=32-bit_rle_rgbe"))
+		{
+			bHasRGBEFormat = true;
+		}
+	}
+
+	if (bHasRGBEFormat)
+	{
+		while (GetHeaderLine(FileDataPtr, Line))
+		{
+			char* HeightStr = FCStringAnsi::Strstr(Line, "Y ");
+			char* WidthStr = FCStringAnsi::Strstr(Line, "X ");
+
+			if (HeightStr != NULL && WidthStr != NULL)
+			{
+				// insert a /0 after the height value
+				*(WidthStr-2) = 0;
+
+				Height = FCStringAnsi::Atoi(HeightStr+2);
+				Width = FCStringAnsi::Atoi(WidthStr+2);
+
+				RGBDataStart = FileDataPtr;
+				return true;
+			}
+		}
+	}
+	else
+	{
+		ErrorMessage = LOCTEXT("WrongFormatError", "The hdr buffer use a unsupported format. Only the 32-bit_rle_rgbe format is supported.");
+	}
+	
+	FreeCompressedData();
+	return false;
+}
+
+bool FHdrImageWrapper::SetCompressed(const void* InCompressedData, int64 InCompressedSize)
+{
+	CompressedDataHolder.Reset(InCompressedSize);
+	CompressedDataHolder.AddUninitialized(InCompressedSize);
+	FMemory::Memcpy(CompressedDataHolder.GetData(), InCompressedData, InCompressedSize);
+
+	return SetCompressedFromView(MakeArrayView(CompressedDataHolder));
+}
+
+bool FHdrImageWrapper::SetRaw(const void* InRawData, int64 InRawSize, const int32 InWidth, const int32 InHeight, const ERGBFormat InFormat, const int32 InBitDepth)
+{
+	// Move the code that export hdr here
+	unimplemented();
+	return false;
+}
+
+bool FHdrImageWrapper::SetAnimationInfo(int32 InNumFrames, int32 InFramerate)
+{
+	unimplemented();
+	return false;
+}
+
+TArray64<uint8> FHdrImageWrapper::GetCompressed(int32 Quality)
+{
+	// Move the code that export hdr here
+	unimplemented();
+	return {};
+}
+
+bool FHdrImageWrapper::GetRaw(const ERGBFormat InFormat, int32 InBitDepth, TArray64<uint8>& OutRawData)
+{
+	if (InFormat != ERGBFormat::BGRE || InBitDepth != 8)
+	{
+		ErrorMessage = LOCTEXT("UnSupportedFormatORBitDepth", "The format and/or the bit depth is not supported by the HdrImageWrapper. Only the BGRE format and a bitdepth of 8 is supported");
+		return false;
+	}
+
+	if (!IsCompressedImageValid())
+	{
+		return false;
+	}
+
+	const int64 SizeRawImageInBytes = Width * Height * 4;
+	OutRawData.Reset(SizeRawImageInBytes);
+	OutRawData.AddUninitialized(SizeRawImageInBytes);
+
+	const uint8* FileDataPtr = RGBDataStart;
+
+	for (int32 Y = 0; Y < Height; ++Y)
+	{
+		if (!DecompressScanline(&(OutRawData[Width * Y * 4]), FileDataPtr))
+		{
+			return false;
+		}
+	}
+
+	return true;
+}
+
+int32 FHdrImageWrapper::GetWidth() const
+{
+	return Width;
+}
+
+int32 FHdrImageWrapper::GetHeight() const
+{
+	return Height;
+}
+
+int32 FHdrImageWrapper::GetBitDepth() const
+{
+	return 8;
+}
+
+ERGBFormat FHdrImageWrapper::GetFormat() const
+{
+	return ERGBFormat::BGRE;
+}
+
+int32 FHdrImageWrapper::GetNumFrames() const
+{
+	return INDEX_NONE;
+}
+
+int32 FHdrImageWrapper::GetFramerate() const
+{
+	return INDEX_NONE;
+}
+
+const FText& FHdrImageWrapper::GetErrorMessage() const
+{
+	return ErrorMessage;
+}
+
+bool FHdrImageWrapper::GetHeaderLine(const uint8*& BufferPos, char Line[256])
+{
+	char* LinePtr = Line;
+
+	const uint8* EndOfBuffer = CompressedData.GetData() + CompressedData.Num();
+	uint32 i;
+
+	for(i = 0; i < 255; ++i)
+	{
+		if(*BufferPos == 0 || *BufferPos == 10 || *BufferPos == 13)
+		{
+			++BufferPos;
+			break;
+		}
+
+		if (BufferPos >= EndOfBuffer)
+		{
+			ErrorMessage = LOCTEXT("RechedEndOfBufferWhileParsingHeader", "Reached the end of the Hdr buffer before we were done reading the header. The Hdr is invalid");
+			return false;
+		}
+
+		*LinePtr++ = *BufferPos++;
+	}
+
+	Line[i] = 0;
+
+	return true;
+}
+
+bool FHdrImageWrapper::DecompressScanline(uint8* Out, const uint8*& In)
+{
+	// minimum and maxmimum scanline length for encoding
+	const int32 MINELEN = 8;
+	const int32 MAXELEN = 0x7fff;
+
+	if(Width < MINELEN || Width > MAXELEN)
+	{
+		return OldDecompressScanline(Out, In, Width);
+	}
+
+	uint8 Red = *In;
+
+	if(Red != 2)
+	{
+		return OldDecompressScanline(Out, In, Width);
+	}
+
+	++In;
+
+	uint8 Green = *In++;
+	uint8 Blue = *In++;
+	uint8 Exponant = *In++;
+
+	if(Green != 2 || Blue & 128)
+	{
+		*Out++ = Blue;
+		*Out++ = Green;
+		*Out++ = Red;
+		*Out++ = Exponant;
+		return OldDecompressScanline(Out, In, Width - 1);
+	}
+
+	const uint8* EndOfOutBuffer = Out + Width * Height;
+	const uint8* EndOfBuffer = CompressedData.GetData() + CompressedData.Num();
+	for(uint32 ChannelRead = 0; ChannelRead < 4; ++ChannelRead)
+	{
+		// The file is in RGBE but we decompress in BGRE So swap the red and blue
+		uint8 CurrentToWrite = ChannelRead;
+		if (ChannelRead == 0)
+		{
+			CurrentToWrite = 2;
+		}
+		else if (ChannelRead == 2)
+		{
+			CurrentToWrite = 0;
+		}
+
+		uint8* OutSingleChannel = Out + CurrentToWrite;
+		int32 MultiRunIndex = 0;
+		while ( MultiRunIndex < Width )
+		{
+			if (In >= EndOfBuffer)
+			{
+				ErrorMessage = UE::ImageWrapper::Private::HdrImageWrapper::GetEndOFBufferErrorMessage();
+				return false;
+			}
+
+			uint8 Current = *In++;
+
+			if (Current > 128)
+			{
+				uint32 Count = Current & 0x7f;
+
+				if (In >= EndOfBuffer)
+				{
+					ErrorMessage = UE::ImageWrapper::Private::HdrImageWrapper::GetEndOFBufferErrorMessage();
+					return false;
+				}
+				Current = *In++;
+
+				if ( OutSingleChannel + 4 * (Count - 1) >= EndOfOutBuffer)
+				{
+					ErrorMessage = UE::ImageWrapper::Private::HdrImageWrapper::GetImageDoneBeforeEndOfBufferErrorMessage();
+					return false;
+				}
+
+				for(uint32 RunIndex = 0; RunIndex < Count; ++RunIndex)
+				{
+					*OutSingleChannel = Current;
+					OutSingleChannel += 4;
+				}
+				MultiRunIndex += Count;
+			}
+			else
+			{
+				uint32 Count = Current;
+
+				if ( OutSingleChannel + 4 * (Count - 1) >= EndOfOutBuffer)
+				{
+					ErrorMessage = UE::ImageWrapper::Private::HdrImageWrapper::GetImageDoneBeforeEndOfBufferErrorMessage();
+					return false;
+				}
+
+				for(uint32 RunIndex = 0; RunIndex < Count; ++RunIndex)
+				{
+					if (In >= EndOfBuffer)
+					{
+						ErrorMessage = UE::ImageWrapper::Private::HdrImageWrapper::GetEndOFBufferErrorMessage();
+						return false;
+					}
+					*OutSingleChannel = *In++;
+					OutSingleChannel += 4;
+				}
+				MultiRunIndex += Count;
+			}
+		}
+	}
+
+	return true;
+}
+
+bool FHdrImageWrapper::OldDecompressScanline(uint8* Out, const uint8*& In, int32 Lenght)
+{
+	int32 Shift = 0;
+
+	while(Lenght > 0)
+	{
+		uint8 Red = *In++; 
+		uint8 Green = *In++;
+		uint8 Blue = *In++; 
+		uint8 Exponent = *In++; 
+
+		if(Red == 1 && Green ==1 && Blue == 1)
+		{
+			int32 Count = ((int32)(Exponent) << Shift);
+			Lenght -= Count;
+
+			if ( Lenght < 0 )
+			{
+				ErrorMessage = LOCTEXT("EndOfLineError","Reached the end of the outputted scanline before finishing decompressing the line. The hdr image is likely to be corrupted");
+				return false;
+			}
+
+			Red = *(Out - 4); 
+			Green = *(Out - 3);
+			Blue = *(Out - 2); 
+			Exponent = *(Out - 1); 
+
+			for(int32 i = Count; i > 0; --i)
+			{
+				*Out++ = Blue;
+				*Out++ = Green;
+				*Out++ = Red;
+				*Out++ = Exponent;
+			}
+
+			Shift += 8;
+		}
+		else
+		{
+			*Out++ = Blue;
+			*Out++ = Green;
+			*Out++ = Red;
+			*Out++ = Exponent;
+			Shift = 0;
+			--Lenght;
+		}
+	}
+
+	return true;
+}
+
+bool FHdrImageWrapper::IsCompressedImageValid() const
+{
+	return CompressedData.Num() > 0 && RGBDataStart;
+}
+
+void FHdrImageWrapper::FreeCompressedData()
+{
+	CompressedData = TArrayView64<const uint8>();
+	RGBDataStart = nullptr;
+	CompressedDataHolder.Empty();
+}
+
+#undef LOCTEXT_NAMESPACE
diff --git a/Engine/Source/Runtime/ImageWrapper/Private/ImageWrapperModule.cpp b/Engine/Source/Runtime/ImageWrapper/Private/ImageWrapperModule.cpp
--- a/Engine/Source/Runtime/ImageWrapper/Private/ImageWrapperModule.cpp
+++ b/Engine/Source/Runtime/ImageWrapper/Private/ImageWrapperModule.cpp
@@ -7,11 +7,13 @@
 
 #include "Formats/BmpImageWrapper.h"
 #include "Formats/ExrImageWrapper.h"
+#include "Formats/HdrImageWrapper.h"
 #include "Formats/IcnsImageWrapper.h"
 #include "Formats/IcoImageWrapper.h"
 #include "Formats/JpegImageWrapper.h"
 #include "Formats/PngImageWrapper.h"
 #include "Formats/TgaImageWrapper.h"
+
 #include "IImageWrapperModule.h"
 
 
@@ -27,6 +29,9 @@ namespace
 	static const uint8 IMAGE_MAGIC_EXR[]  = {0x76, 0x2F, 0x31, 0x01};
 	static const uint8 IMAGE_MAGIC_ICNS[] = {0x69, 0x63, 0x6E, 0x73};
 
+	// Binary for #?RADIANCE
+	static const uint8 IMAGE_MAGIC_HDR[] = {0x23, 0x3f, 0x52, 0x41, 0x44, 0x49, 0x41, 0x4e, 0x43, 0x45, 0x0a};
+
 	/** Internal helper function to verify image signature. */
 	template <int32 MagicCount> bool StartsWith(const uint8* Content, int64 ContentSize, const uint8 (&Magic)[MagicCount])
 	{
@@ -101,6 +106,8 @@ public:
 		case EImageFormat::TGA:
 			ImageWrapper = MakeShared<FTgaImageWrapper>();
 			break;
+		case EImageFormat::HDR:
+			ImageWrapper = MakeShared<FHdrImageWrapper>();
 
 		default:
 			break;
@@ -136,6 +143,10 @@ public:
 		{
 			Format = EImageFormat::ICNS;
 		}
+		else if (StartsWith((uint8*)CompressedData, CompressedSize, IMAGE_MAGIC_HDR))
+		{
+			Format = EImageFormat::HDR;
+		}
 
 		return Format;
 	}
diff --git a/Engine/Source/Runtime/ImageWrapper/Public/Formats/HdrImageWrapper.h b/Engine/Source/Runtime/ImageWrapper/Public/Formats/HdrImageWrapper.h
new file mode 100644
--- /dev/null
+++ b/Engine/Source/Runtime/ImageWrapper/Public/Formats/HdrImageWrapper.h
@@ -0,0 +1,61 @@
+// Copyright Epic Games, Inc. All Rights Reserved.
+
+#pragma once
+
+#include "IImageWrapper.h"
+#include "CoreMinimal.h"
+
+// http://radsite.lbl.gov/radiance/refer/Notes/picture_format.html
+// http://paulbourke.net/dataformats/pic/
+
+/** To load the HDR file image format. Does not support all possible types HDR formats (e.g. xyze is not supported) */
+class IMAGEWRAPPER_API FHdrImageWrapper : public IImageWrapper
+{
+public:
+
+	// Todo we should have this for all image wrapper.
+	bool SetCompressedFromView(TArrayView64<const uint8> Data);
+
+	// IIMageWrapper Interface begin
+	virtual bool SetCompressed(const void* InCompressedData, int64 InCompressedSize) override;
+	virtual bool SetRaw(const void* InRawData, int64 InRawSize, const int32 InWidth, const int32 InHeight, const ERGBFormat InFormat, const int32 InBitDepth) override;
+	virtual bool SetAnimationInfo(int32 InNumFrames, int32 InFramerate) override;
+	virtual TArray64<uint8> GetCompressed(int32 Quality = 0) override;
+	virtual bool GetRaw(const ERGBFormat InFormat, int32 InBitDepth, TArray64<uint8>& OutRawData) override;
+
+	virtual int32 GetWidth() const override;
+	virtual int32 GetHeight() const override;
+	virtual int32 GetBitDepth() const override;
+	virtual ERGBFormat GetFormat() const override;
+
+	virtual int32 GetNumFrames() const override;
+	virtual int32 GetFramerate() const override;
+	// IImageWrapper Interface end
+
+	const FText& GetErrorMessage() const;
+
+	void FreeCompressedData();
+
+private:
+	bool GetHeaderLine(const uint8*& BufferPos, char Line[256]);
+
+	/** @param Out order in bytes: RGBE */
+	bool DecompressScanline(uint8* Out, const uint8*& In);
+
+	bool OldDecompressScanline(uint8* Out, const uint8*& In, int32 Lenght);
+
+	bool IsCompressedImageValid() const;
+
+	TArrayView64<const uint8> CompressedData;
+	const uint8* RGBDataStart = nullptr;
+
+	TArray64<uint8> CompressedDataHolder;
+
+	/** INDEX_NONE if not valid */
+	int32 Width = INDEX_NONE;
+	/** INDEX_NONE if not valid */
+	int32 Height = INDEX_NONE;
+
+	// Reported error
+	FText ErrorMessage;
+};
\ No newline at end of file
diff --git a/Engine/Source/Runtime/ImageWrapper/Public/IImageWrapper.h b/Engine/Source/Runtime/ImageWrapper/Public/IImageWrapper.h
--- a/Engine/Source/Runtime/ImageWrapper/Public/IImageWrapper.h
+++ b/Engine/Source/Runtime/ImageWrapper/Public/IImageWrapper.h
@@ -37,7 +37,10 @@ enum class EImageFormat : int8
 	ICNS,
 	
 	/** Truevision TGA / TARGA */
-	TGA
+	TGA,
+
+	/** Hdr file from radiance using RGBE*/
+	HDR,
 };
 
 
@@ -59,6 +62,9 @@ enum class ERGBFormat : int8
 
 	// Red, Green, Blue and Alpha using IEEE Floating-Point Arithmetic (see IEEE754). The format is always binary.
 	RGBAF = 3,
+
+	// Blue, Green, Red and Exponent (Similar to the RGBE format from radiance but with the blue and red channel inversed)
+	BGRE = 4,
 };
 
 
