Avoid buffer overflow when using JPEG compression for small images.

diff --git a/Engine/Source/Runtime/ImageWrapper/Private/Formats/JpegImageWrapper.cpp b/Engine/Source/Runtime/ImageWrapper/Private/Formats/JpegImageWrapper.cpp
--- a/Engine/Source/Runtime/ImageWrapper/Private/Formats/JpegImageWrapper.cpp
+++ b/Engine/Source/Runtime/ImageWrapper/Private/Formats/JpegImageWrapper.cpp
@@ -280,18 +280,18 @@ void FJpegImageWrapper::CompressTurbo(int32 Quality)
 		check(Width > 0);
 		check(Height > 0);
 
-		CompressedData.Reset(RawData.Num());
-		CompressedData.AddUninitialized(RawData.Num());
-
 		const int PixelFormat = ConvertTJpegPixelFormat(RawFormat);
-		unsigned char* OutBuffer = CompressedData.GetData();
-		unsigned long OutBufferSize = static_cast<unsigned long>(CompressedData.Num());
+		const int Subsampling = TJSAMP_420;
 		const int Flags = TJFLAG_NOREALLOC | TJFLAG_FASTDCT;
 
-		const bool bSuccess = tjCompress2(Compressor, RawData.GetData(), Width, 0, Height, PixelFormat, &OutBuffer, &OutBufferSize, TJSAMP_420, Quality, Flags) == 0;
+		unsigned long OutBufferSize = tjBufSize(Width, Height, Subsampling);
+		CompressedData.SetNum(OutBufferSize);
+		unsigned char* OutBuffer = CompressedData.GetData();
+
+		const bool bSuccess = tjCompress2(Compressor, RawData.GetData(), Width, 0, Height, PixelFormat, &OutBuffer, &OutBufferSize, Subsampling, Quality, Flags) == 0;
 		check(bSuccess);
 
-		CompressedData.RemoveAt((int64)OutBufferSize, CompressedData.Num() - (int64)OutBufferSize);
+		CompressedData.SetNum(OutBufferSize);
 	}
 }
 
