Add full support for 32-bit float EXR image channels.

diff --git a/Engine/Plugins/Experimental/PanoramicCapture/Source/PanoramicCapture/Private/SceneCapturer.cpp b/Engine/Plugins/Experimental/PanoramicCapture/Source/PanoramicCapture/Private/SceneCapturer.cpp
--- a/Engine/Plugins/Experimental/PanoramicCapture/Source/PanoramicCapture/Private/SceneCapturer.cpp
+++ b/Engine/Plugins/Experimental/PanoramicCapture/Source/PanoramicCapture/Private/SceneCapturer.cpp
@@ -1283,7 +1283,7 @@ void USceneCapturer::Tick( float DeltaTime )
 				// write
 				TSharedPtr<IImageWrapper> ImageWrapper = ImageWrapperModule.CreateImageWrapper(EImageFormat::EXR);
 				ImageWrapper->SetRaw(CombinedAtlas.GetData(), CombinedAtlas.GetAllocatedSize(), SphericalAtlasWidth, SphericalAtlasHeight * EyeCount, ERGBFormat::RGBAF, 32);
-				const TArray64<uint8> ImageData = ImageWrapper->GetCompressed((int32)EImageCompressionQuality::Default);
+				const TArray64<uint8> ImageData = ImageWrapper->GetCompressed();
 				FFileHelper::SaveArrayToFile(ImageData, *AtlasName);
 				ImageWrapper.Reset();
 				FMessageLog(StereoPanoramaLogName).Message(EMessageSeverity::Info, LOCTEXT("Done", "Done!"));
diff --git a/Engine/Plugins/MovieScene/MovieRenderPipeline/Source/MovieRenderPipelineCore/Private/MoviePipelineRendering.cpp b/Engine/Plugins/MovieScene/MovieRenderPipeline/Source/MovieRenderPipelineCore/Private/MoviePipelineRendering.cpp
--- a/Engine/Plugins/MovieScene/MovieRenderPipeline/Source/MovieRenderPipelineCore/Private/MoviePipelineRendering.cpp
+++ b/Engine/Plugins/MovieScene/MovieRenderPipeline/Source/MovieRenderPipelineCore/Private/MoviePipelineRendering.cpp
@@ -466,9 +466,9 @@ void UMoviePipeline::OnSampleRendered(TUniquePtr<FImagePixelData>&& OutputSample
 
 	FImagePixelDataPayload* InFrameData = OutputSample->GetPayload<FImagePixelDataPayload>();
 	TileImageTask->Format = EImageFormat::EXR;
-	TileImageTask->CompressionQuality = 100;
+	TileImageTask->CompressionQuality = (int32)EImageCompressionQuality::Default;
 
-	FString OutputName = FString::Printf(TEXT("/%s_SS_%d_TS_%d_TileX_%d_TileY_%d.%d.jpeg"),
+	FString OutputName = FString::Printf(TEXT("/%s_SS_%d_TS_%d_TileX_%d_TileY_%d.%d.exr"),
 		*InFrameData->PassIdentifier.Name, InFrameData->SampleState.SpatialSampleIndex, InFrameData->SampleState.TemporalSampleIndex,
 		InFrameData->SampleState.TileIndexes.X, InFrameData->SampleState.TileIndexes.Y, InFrameData->SampleState.OutputState.OutputFrameNumber);
 
diff --git a/Engine/Plugins/MovieScene/MovieRenderPipeline/Source/MovieRenderPipelineRenderPasses/Private/MoviePipelineEXROutput.cpp b/Engine/Plugins/MovieScene/MovieRenderPipeline/Source/MovieRenderPipelineRenderPasses/Private/MoviePipelineEXROutput.cpp
--- a/Engine/Plugins/MovieScene/MovieRenderPipeline/Source/MovieRenderPipelineRenderPasses/Private/MoviePipelineEXROutput.cpp
+++ b/Engine/Plugins/MovieScene/MovieRenderPipeline/Source/MovieRenderPipelineRenderPasses/Private/MoviePipelineEXROutput.cpp
@@ -266,6 +266,7 @@ static FString GetChannelName(const FString& InLayerName, const int32 InChannelI
 		}
 		break;
 		case ERGBFormat::Gray:
+		case ERGBFormat::GrayF:
 		{
 			check(InChannelIndex < UE_ARRAY_COUNT(GrayChannelNames));
 			ChannelNames = GrayChannelNames;
diff --git a/Engine/Plugins/Runtime/Interchange/Source/Import/Private/Texture/InterchangeEXRTranslator.cpp b/Engine/Plugins/Runtime/Interchange/Source/Import/Private/Texture/InterchangeEXRTranslator.cpp
--- a/Engine/Plugins/Runtime/Interchange/Source/Import/Private/Texture/InterchangeEXRTranslator.cpp
+++ b/Engine/Plugins/Runtime/Interchange/Source/Import/Private/Texture/InterchangeEXRTranslator.cpp
@@ -83,26 +83,16 @@ TOptional<UE::Interchange::FImportImage> UInterchangeEXRTranslator::GetTexturePa
 
 	if (!UE::Interchange::FImportImageHelper::IsImportResolutionValid(Width, Height, bAllowNonPowerOfTwo))
 	{
-		UE_LOG(LogInterchangeImport, Error, TEXT("Failed to import PCX, invalid resolution. Resolution[%d, %d], AllowPowerOfTwo[%s], [%s]"), Width, Height, bAllowNonPowerOfTwo ? TEXT("True") : TEXT("false"), *Filename);
+		UE_LOG(LogInterchangeImport, Error, TEXT("Failed to import EXR, invalid resolution. Resolution[%d, %d], AllowPowerOfTwo[%s], [%s]"), Width, Height, bAllowNonPowerOfTwo ? TEXT("True") : TEXT("false"), *Filename);
 		return TOptional<UE::Interchange::FImportImage>();
 	}
 
-	// Select the texture's source format
-	ETextureSourceFormat TextureFormat = TSF_Invalid;
-	int32 BitDepth = ExrImageWrapper->GetBitDepth();
-	ERGBFormat Format = ExrImageWrapper->GetFormat();
-
-	if (Format == ERGBFormat::RGBA && BitDepth == 16)
-	{
-		TextureFormat = TSF_RGBA16F;
-		Format = ERGBFormat::BGRA;
-	}
-
-	if (TextureFormat == TSF_Invalid)
-	{
-		UE_LOG(LogInterchangeImport, Error, TEXT("EXR file [%s] contains data in an unsupported format"), *Filename);
-		return TOptional<UE::Interchange::FImportImage>();
-	}
+	// Currently the only texture source format compatible with EXR image formats is TSF_RGBA16F.
+	// EXR decoder automatically converts imported image channels into the requested float format.
+	// In case if the imported image is a gray scale image, its content will be stored in the green channel of the created texture.
+	ETextureSourceFormat TextureFormat = TSF_RGBA16F;
+	ERGBFormat Format = ERGBFormat::RGBAF;
+	int32 BitDepth = 16;
 
 	UE::Interchange::FImportImage PayloadData;
 	PayloadData.Init2DWithParams(
diff --git a/Engine/Source/Editor/UnrealEd/Private/Factories/EditorFactories.cpp b/Engine/Source/Editor/UnrealEd/Private/Factories/EditorFactories.cpp
--- a/Engine/Source/Editor/UnrealEd/Private/Factories/EditorFactories.cpp
+++ b/Engine/Source/Editor/UnrealEd/Private/Factories/EditorFactories.cpp
@@ -3191,21 +3191,12 @@ bool UTextureFactory::ImportImage(const uint8* Buffer, uint32 Length, FFeedbackC
 				return false;
 			}
 
-			// Select the texture's source format
-			ETextureSourceFormat TextureFormat = TSF_Invalid;
-			int32 BitDepth = ExrImageWrapper->GetBitDepth();
-			ERGBFormat Format = ExrImageWrapper->GetFormat();
-
-			if (Format == ERGBFormat::RGBAF && BitDepth == 16)
-			{
-				TextureFormat = TSF_RGBA16F;
-			}
-
-			if (TextureFormat == TSF_Invalid)
-			{
-				Warn->Logf(ELogVerbosity::Error, TEXT("EXR file contains data in an unsupported format."));
-				return false;
-			}
+			// Currently the only texture source format compatible with EXR image formats is TSF_RGBA16F.
+			// EXR decoder automatically converts imported image channels into the requested float format.
+			// In case if the imported image is a gray scale image, its content will be stored in the green channel of the created texture.
+			ETextureSourceFormat TextureFormat = TSF_RGBA16F;
+			ERGBFormat Format = ERGBFormat::RGBAF;
+			int32 BitDepth = 16;
 
 			OutImage.Init2DWithParams(
 				Width,
diff --git a/Engine/Source/Runtime/Engine/Private/ImageUtils.cpp b/Engine/Source/Runtime/Engine/Private/ImageUtils.cpp
--- a/Engine/Source/Runtime/Engine/Private/ImageUtils.cpp
+++ b/Engine/Source/Runtime/Engine/Private/ImageUtils.cpp
@@ -822,7 +822,7 @@ ENGINE_API bool FImageUtils::ExportRenderTarget2DAsEXR(UTextureRenderTarget2D* T
 		bSuccess = GetRawData(TexRT, RawData);
 
 		int32 BitsPerPixel = TexRT->GetFormat() == PF_B8G8R8A8 ? 8 : (sizeof(FFloat16Color) / 4) * 8;
-		ERGBFormat RGBFormat = TexRT->GetFormat() == PF_B8G8R8A8 ? ERGBFormat::BGRA : ERGBFormat::RGBA;
+		ERGBFormat RGBFormat = TexRT->GetFormat() == PF_B8G8R8A8 ? ERGBFormat::BGRA : ERGBFormat::RGBAF;
 
 		IImageWrapperModule& ImageWrapperModule = FModuleManager::Get().LoadModuleChecked<IImageWrapperModule>(TEXT("ImageWrapper"));
 
@@ -830,7 +830,7 @@ ENGINE_API bool FImageUtils::ExportRenderTarget2DAsEXR(UTextureRenderTarget2D* T
 
 		EXRImageWrapper->SetRaw(RawData.GetData(), RawData.GetAllocatedSize(), Size.X, Size.Y, RGBFormat, BitsPerPixel);
 
-		const TArray64<uint8> Data = EXRImageWrapper->GetCompressed(100);
+		const TArray64<uint8> Data = EXRImageWrapper->GetCompressed();
 
 		Ar.Serialize((void*)Data.GetData(), Data.GetAllocatedSize());
 
@@ -950,7 +950,7 @@ UTexture2D* FImageUtils::ImportBufferAsTexture2D(TArrayView64<const uint8> Buffe
 			if (BitDepth == 16)
 			{
 				PixelFormat = PF_FloatRGBA;
-				RGBFormat = ERGBFormat::RGBA;
+				RGBFormat = ERGBFormat::RGBAF;
 			}
 			else if (BitDepth == 8)
 			{
diff --git a/Engine/Source/Runtime/ImageWrapper/Private/Formats/ExrImageWrapper.cpp b/Engine/Source/Runtime/ImageWrapper/Private/Formats/ExrImageWrapper.cpp
--- a/Engine/Source/Runtime/ImageWrapper/Private/Formats/ExrImageWrapper.cpp
+++ b/Engine/Source/Runtime/ImageWrapper/Private/Formats/ExrImageWrapper.cpp
@@ -11,38 +11,9 @@
 
 FExrImageWrapper::FExrImageWrapper()
 	: FImageWrapperBase()
-	, bUseCompression(true)
 {
 }
 
-template <typename sourcetype> class FSourceImageRaw
-{
-public:
-	FSourceImageRaw(const TArray64<uint8>& SourceImageBitmapIN, uint64 ChannelsIN, uint64 WidthIN, uint64 HeightIN)
-		: SourceImageBitmap(SourceImageBitmapIN),
-			Width( WidthIN ),
-			Height(HeightIN),
-			Channels(ChannelsIN)
-	{		
-		check(SourceImageBitmap.Num() == Channels*Width*Height*sizeof(sourcetype));
-	}
-
-	uint64 GetXStride() const	{return sizeof(sourcetype) * Channels;}
-	uint64 GetYStride() const	{return GetXStride() * Width;}
-
-	const sourcetype* GetHorzLine(uint64 y, uint64 Channel)
-	{
-		return &SourceImageBitmap[(sizeof(sourcetype) * Channel) + (GetYStride() * y)];
-	}
-
-private:
-	const TArray64<uint8>& SourceImageBitmap;
-	uint64 Width;
-	uint64 Height;
-	uint64 Channels;
-};
-
-
 class FMemFileOut : public Imf::OStream
 {
 public:
@@ -177,290 +148,313 @@ private:
 	int64 Pos;
 };
 
+const char* cChannelNamesRGBA[] = { "R", "G", "B", "A" };
+const char* cChannelNamesBGRA[] = { "B", "G", "R", "A" };
+const char* cChannelNamesGray[] = { "G" };
 
-namespace
+int32 GetChannelNames(ERGBFormat InRGBFormat, const char* const*& OutChannelNames)
 {
-	/////////////////////////////////////////
-	// 8 bit per channel source
-	void ExtractAndConvertChannel(const uint8*Src, uint64 SrcChannels, uint64 x, uint64 y, float* ChannelOUT)
+	int32 ChannelCount;
+
+	switch (InRGBFormat)
 	{
-		for (uint64 i = 0; i < uint64(x)*uint64(y); i++, Src += SrcChannels)
-		{
-			ChannelOUT[i] = *Src / 255.f;
-		}
+	case ERGBFormat::RGBA:
+	case ERGBFormat::RGBAF:
+		OutChannelNames = cChannelNamesRGBA;
+		ChannelCount = UE_ARRAY_COUNT(cChannelNamesRGBA);
+		break;
+
+	case ERGBFormat::BGRA:
+		OutChannelNames = cChannelNamesBGRA;
+		ChannelCount = UE_ARRAY_COUNT(cChannelNamesBGRA);
+		break;
+
+	case ERGBFormat::Gray:
+	case ERGBFormat::GrayF:
+		OutChannelNames = cChannelNamesGray;
+		ChannelCount = UE_ARRAY_COUNT(cChannelNamesGray);
+		break;
+
+	default:
+		OutChannelNames = nullptr;
+		ChannelCount = 0;
 	}
 
-	void ExtractAndConvertChannel(const uint8*Src, uint64 SrcChannels, uint64 x, uint64 y, FFloat16* ChannelOUT)
+	return ChannelCount;
+}
+
+bool FExrImageWrapper::SetRaw(const void* InRawData, int64 InRawSize, const int32 InWidth, const int32 InHeight, const ERGBFormat InFormat, const int32 InBitDepth)
+{
+	check(InRawData);
+	check(InRawSize > 0);
+	check(InWidth > 0);
+	check(InHeight > 0);
+
+	switch (InBitDepth)
 	{
-		for (uint64 i = 0; i < uint64(x)*uint64(y); i++, Src += SrcChannels)
+	case 8:
+		if (InFormat != ERGBFormat::RGBA && InFormat != ERGBFormat::BGRA && InFormat != ERGBFormat::Gray)
 		{
-			ChannelOUT[i] = FFloat16(*Src / 255.f);
+			return false;
 		}
-	}
+		break;
 
-	/////////////////////////////////////////
-	// 16 bit per channel source
-	void ExtractAndConvertChannel(const FFloat16*Src, uint64 SrcChannels, uint64 x, uint64 y, float* ChannelOUT)
-	{
-		for (uint64 i = 0; i < uint64(x)*uint64(y); i++, Src += SrcChannels)
+	case 16:
+	case 32:
+		if (InFormat == ERGBFormat::RGBA || InFormat == ERGBFormat::Gray)
 		{
-			ChannelOUT[i] = Src->GetFloat();
+			// Before ERGBFormat::RGBAF and ERGBFormat::GrayF were introduced, ERGBFormat::RGBA and ERGBFormat::Gray were used to describe float pixel formats.
+			// ERGBFormat::RGBA and ERGBFormat::Gray should now only be used for integer channels.
+			// Note that EXR uint32 compression is currently not supported.
+			const TCHAR* FormatName = (InFormat == ERGBFormat::RGBA) ? TEXT("RGBA") : TEXT("Gray");
+			UE_LOG(LogImageWrapper, Warning, TEXT("Usage of 16-bit and 32-bit ERGBFormat::%s raw format for compressing EXR images is deprecated, if you are compressing float channels please specify ERGBFormat::%sF instead."), *FormatName, *FormatName);
 		}
-	}
-
-	void ExtractAndConvertChannel(const FFloat16*Src, uint64 SrcChannels, uint64 x, uint64 y, FFloat16* ChannelOUT)
-	{
-		for (uint64 i = 0; i < uint64(x)*uint64(y); i++, Src += SrcChannels)
+		if (InFormat != ERGBFormat::RGBAF && InFormat != ERGBFormat::GrayF)
 		{
-			ChannelOUT[i] = *Src;
+			return false;
 		}
+		break;
 	}
 
-	/////////////////////////////////////////
-	// 32 bit per channel source
-	void ExtractAndConvertChannel(const float* Src, uint64 SrcChannels, uint64 x, uint64 y, float* ChannelOUT)
+	return FImageWrapperBase::SetRaw(InRawData, InRawSize, InWidth, InHeight, InFormat, InBitDepth);
+}
+
+bool FExrImageWrapper::SetCompressed(const void* InCompressedData, int64 InCompressedSize)
+{
+	check(InCompressedData);
+	check(InCompressedSize > 0);
+
+	// Check the magic value in advance to avoid spamming the log with EXR parsing errors.
+	if (InCompressedSize < sizeof(uint32) || *(uint32*)InCompressedData != Imf::MAGIC)
 	{
-		for (uint64 i = 0; i < uint64(x)*uint64(y); i++, Src += SrcChannels)
-		{
-			ChannelOUT[i] = *Src;
-		}
+		return false;
 	}
 
-	void ExtractAndConvertChannel(const float* Src, uint64 SrcChannels, uint64 x, uint64 y, FFloat16* ChannelOUT)
+	if (!FImageWrapperBase::SetCompressed(InCompressedData, InCompressedSize))
 	{
-		for (uint64 i = 0; i < uint64(x)*uint64(y); i++, Src += SrcChannels)
-		{
-			ChannelOUT[i] = *Src;
-		}
+		return false;
 	}
 
-	/////////////////////////////////////////
-	int32 GetNumChannelsFromFormat(ERGBFormat Format)
+	// openEXR can throw exceptions when parsing invalid data.
+	try
 	{
-		switch (Format)
+		FMemFileIn MemFile(CompressedData.GetData(), CompressedData.Num());
+		Imf::InputFile ImfFile(MemFile);
+		Imf::Header ImfHeader = ImfFile.header();
+		Imf::ChannelList ImfChannels = ImfHeader.channels();
+		Imath::Box2i ImfDataWindow = ImfHeader.dataWindow();
+		Width = ImfDataWindow.max.x - ImfDataWindow.min.x + 1;
+		Height = ImfDataWindow.max.y - ImfDataWindow.min.y + 1;
+
+		bool bHasOnlyHALFChannels = true;
+		bool bMatchesGrayOrder = true;
+		int32 ChannelCount = 0;
+
+		for (Imf::ChannelList::Iterator Iter = ImfChannels.begin(); Iter != ImfChannels.end(); ++Iter, ++ChannelCount)
 		{
-			case ERGBFormat::RGBA:
-			case ERGBFormat::BGRA:
-			case ERGBFormat::RGBAF:
-				return 4;
-			case ERGBFormat::Gray:
-				return 1;
+			bHasOnlyHALFChannels = bHasOnlyHALFChannels && Iter.channel().type == Imf::HALF;
+			bMatchesGrayOrder = bMatchesGrayOrder && ChannelCount < UE_ARRAY_COUNT(cChannelNamesGray) && !strcmp(Iter.name(), cChannelNamesGray[ChannelCount]);
 		}
-		checkNoEntry();
-		return 1;
+
+		BitDepth = (ChannelCount && bHasOnlyHALFChannels) ? 16 : 32;
+
+		// EXR uint32 channels are currently not supported, therefore input channels are always treated as float channels.
+		// Channel combinations which don't match the ERGBFormat::GrayF pattern are qualified as ERGBFormat::RGBAF.
+		// Note that channels inside the EXR file are indexed by name, therefore can be decoded in any RGB order.
+		Format = (ChannelCount == UE_ARRAY_COUNT(cChannelNamesGray) && bMatchesGrayOrder) ? ERGBFormat::GrayF : ERGBFormat::RGBAF;
 	}
+	catch (const std::exception& Exception)
+	{
+		TStringConversion<TStringConvert<char, TCHAR>> Convertor(Exception.what());
+		UE_LOG(LogImageWrapper, Error, TEXT("Cannot parse EXR image header: %s"), Convertor.Get());
+		SetError(Convertor.Get());
+		return false;
+	}
+
+	return true;
 }
 
-const char* FExrImageWrapper::GetRawChannelName(int ChannelIndex) const
+void FExrImageWrapper::Compress(int32 Quality)
 {
-	const int32 MaxChannels = 4;
-	static const char* RGBAChannelNames[] = { "R", "G", "B", "A" };
-	static const char* BGRAChannelNames[] = { "B", "G", "R", "A" };
-	static const char* GrayChannelNames[] = { "G" };
-	check(ChannelIndex < MaxChannels);
+	check(RawData.Num());
 
-	const char** ChannelNames = BGRAChannelNames;
+	// Ensure we haven't already compressed the file.
+	if (CompressedData.Num())
+	{
+		return;
+	}
 
-	switch (RawFormat)
+	const double StartTime = FPlatformTime::Seconds();
+
+	Imf::PixelType ImfPixelType;
+	TArray64<uint8> ConvertedRawData;
+	bool bNeedsConversion = false;
+
+	if (RawBitDepth == 8)
 	{
-		case ERGBFormat::RGBA:
-		case ERGBFormat::RGBAF:
-		{
-			ChannelNames = RGBAChannelNames;
-		}
-		break;
-		case ERGBFormat::BGRA:
-		{
-			ChannelNames = BGRAChannelNames;
-		}
-		break;
-		case ERGBFormat::Gray:
+		// uint8 channels are linearly converted into FFloat16 channels.
+		ConvertedRawData.SetNumUninitialized(sizeof(FFloat16) * RawData.Num());
+		FFloat16* Output = reinterpret_cast<FFloat16*>(ConvertedRawData.GetData());
+		for (int64 i = 0; i < RawData.Num(); ++i)
 		{
-			check(ChannelIndex < UE_ARRAY_COUNT(GrayChannelNames));
-			ChannelNames = GrayChannelNames;
+			Output[i] = FFloat16(RawData[i] / 255.f);
 		}
-		break;
-		default:
-			checkNoEntry();
+		ImfPixelType = Imf::HALF;
+		bNeedsConversion = true;
+	}
+	else
+	{
+		ImfPixelType = (RawBitDepth == 16) ? Imf::HALF : Imf::FLOAT;
 	}
-	return ChannelNames[ChannelIndex];
-}
 
-template <typename Imf::PixelType OutputFormat>
-struct TExrImageOutputChannelType;
+	const TArray64<uint8>& PixelData = bNeedsConversion ? ConvertedRawData : RawData;
 
-template <> struct TExrImageOutputChannelType<Imf::HALF>  { typedef FFloat16 Type; };
-template <> struct TExrImageOutputChannelType<Imf::FLOAT> { typedef float Type; };
+	const char* const* ChannelNames;
+	int32 ChannelCount = GetChannelNames(RawFormat, ChannelNames);
+	check((int64)ChannelCount * Width * Height * RawBitDepth == RawData.Num() * 8);
 
-template <Imf::PixelType OutputFormat, typename sourcetype>
-void FExrImageWrapper::WriteFrameBufferChannel(Imf::FrameBuffer& ImfFrameBuffer, const char* ChannelName, const sourcetype* SrcData, TArray64<uint8>& ChannelBuffer)
-{
-	const int64 OutputPixelSize = ((OutputFormat == Imf::HALF) ? 2 : 4);
-	ChannelBuffer.AddUninitialized(uint64(Width)*uint64(Height)*uint64(OutputPixelSize));
-	uint64 SrcChannels = GetNumChannelsFromFormat(RawFormat);
-	ExtractAndConvertChannel(SrcData, SrcChannels, Width, Height, (typename TExrImageOutputChannelType<OutputFormat>::Type*)&ChannelBuffer[0]);
-	Imf::Slice FrameChannel = Imf::Slice(OutputFormat, (char*)&ChannelBuffer[0], OutputPixelSize, uint64(Width)*uint64(OutputPixelSize));
-	ImfFrameBuffer.insert(ChannelName, FrameChannel);
-}
+	int32 BytesPerChannelPixel = (ImfPixelType == Imf::HALF) ? 2 : 4;
+	TArray<TArray64<uint8>> ChannelData;
+	ChannelData.SetNum(ChannelCount);
 
-template <Imf::PixelType OutputFormat, typename sourcetype>
-void FExrImageWrapper::CompressRaw(const sourcetype* SrcData, bool bIgnoreAlpha)
-{
-	const double StartTime = FPlatformTime::Seconds();
-	uint32 NumWriteComponents = GetNumChannelsFromFormat(RawFormat);
-	if (bIgnoreAlpha && NumWriteComponents == 4)
+	for (int32 c = 0; c < ChannelCount; ++c)
 	{
-		NumWriteComponents = 3;
+		ChannelData[c].SetNumUninitialized((int64)BytesPerChannelPixel * Width * Height);
 	}
 
-	Imf::Compression Comp = bUseCompression ? Imf::Compression::ZIP_COMPRESSION : Imf::Compression::NO_COMPRESSION;
-	Imf::Header Header(Width, Height, 1, Imath::V2f(0, 0), 1, Imf::LineOrder::INCREASING_Y, Comp);
-	
-	for (uint32 Channel = 0; Channel < NumWriteComponents; Channel++)
+	// EXR channels are compressed non-interleaved.
+	for (int64 OffsetNonInterleaved = 0, OffsetInterleaved = 0; OffsetInterleaved < PixelData.Num(); OffsetNonInterleaved += BytesPerChannelPixel)
 	{
-		Header.channels().insert(GetRawChannelName(Channel), Imf::Channel(OutputFormat));
+		for (int32 c = 0; c < ChannelCount; ++c)
+		{
+			for (int32 b = 0; b < BytesPerChannelPixel; ++b, ++OffsetInterleaved)
+			{
+				ChannelData[c][OffsetNonInterleaved + b] = PixelData[OffsetInterleaved];
+			}
+		}
 	}
 
-	// OutputFormat is 0 UINT, 1 HALF or 2 FLOAT. Size evaluates to 2/4/8 or 1/2/4 compressed
-	const int32 OutputPixelSize = (bUseCompression ? 1 : 2) << OutputFormat;
-	checkSlow(OutputPixelSize >= 1 && OutputPixelSize <= 8);
-
-	FMemFileOut MemFile("");
+	Imf::Compression ImfCompression = (Quality == (int32)EImageCompressionQuality::Uncompressed) ? Imf::Compression::NO_COMPRESSION : Imf::Compression::ZIP_COMPRESSION;
+	Imf::Header ImfHeader(Width, Height, 1, Imath::V2f(0, 0), 1, Imf::LineOrder::INCREASING_Y, ImfCompression);
 	Imf::FrameBuffer ImfFrameBuffer;
-	TArray64<uint8> ChannelOutputBuffers[4];
 
-	for (uint32 Channel = 0; Channel < NumWriteComponents; Channel++)
+	for (int32 c = 0; c < ChannelCount; ++c)
 	{
-		WriteFrameBufferChannel<OutputFormat>(ImfFrameBuffer, GetRawChannelName(Channel), SrcData + Channel, ChannelOutputBuffers[Channel]);
+		ImfHeader.channels().insert(ChannelNames[c], Imf::Channel(ImfPixelType));
+		ImfFrameBuffer.insert(ChannelNames[c], Imf::Slice(ImfPixelType, (char*)ChannelData[c].GetData(), BytesPerChannelPixel, (size_t)BytesPerChannelPixel * Width));
 	}
 
-	int64 FileLength;
+	FMemFileOut MemFile("");
+	int64 MemFileLength;
+
 	{
 		// This scope ensures that IMF::Outputfile creates a complete file by closing the file when it goes out of scope.
 		// To complete the file, EXR seeks back into the file and writes the scanline offsets when the file is closed, 
 		// which moves the tellp location. So file length is stored in advance for later use.
 
-		Imf::OutputFile ImfFile(MemFile, Header, FPlatformMisc::NumberOfCoresIncludingHyperthreads());
+		Imf::OutputFile ImfFile(MemFile, ImfHeader, FPlatformMisc::NumberOfCoresIncludingHyperthreads());
 		ImfFile.setFrameBuffer(ImfFrameBuffer);
-	
-		MemFile.Data.AddUninitialized(int64(Width) * int64(Height) * int64(NumWriteComponents) * int64(OutputFormat == 2 ? 4 : 2));
 		ImfFile.writePixels(Height);
-		FileLength = MemFile.tellp();
+		MemFileLength = MemFile.tellp();
 	}
 
 	CompressedData = MoveTemp(MemFile.Data);
-	CompressedData.SetNum(FileLength);
+	CompressedData.SetNum(MemFileLength);
 
 	const double DeltaTime = FPlatformTime::Seconds() - StartTime;
 	UE_LOG(LogImageWrapper, Verbose, TEXT("Compressed image in %.3f seconds"), DeltaTime);
 }
 
-void FExrImageWrapper::Compress( int32 Quality )
+void FExrImageWrapper::Uncompress(const ERGBFormat InFormat, const int32 InBitDepth)
 {
-	check(RawData.Num() != 0);
-	check(Width > 0);
-	check(Height > 0);
-	check(RawBitDepth == 8 || RawBitDepth == 16 || RawBitDepth == 32);
+	check(CompressedData.Num());
 
-	const int32 MaxComponents = 4;
+	// Ensure we haven't already uncompressed the file.
+	if (RawData.Num() && InFormat == RawFormat && InBitDepth == RawBitDepth)
+	{
+		return;
+	}
 
-	bUseCompression = (Quality != (int32)EImageCompressionQuality::Uncompressed);
+	FString ErrorMessage;
 
-	switch (RawBitDepth)
+	if (InBitDepth == 16 && (InFormat == ERGBFormat::RGBA || InFormat == ERGBFormat::Gray))
 	{
-		case 8:
-			CompressRaw<Imf::HALF>(&RawData[0], false);
-		break;
-		case 16:
-			CompressRaw<Imf::HALF>((const FFloat16*)&RawData[0], false);
-		break;
-		case 32:
-			CompressRaw<Imf::FLOAT>((const float*)&RawData[0], false);
-		break;
-		default:
-			checkNoEntry();
+		// Before ERGBFormat::RGBAF and ERGBFormat::GrayF were introduced, 16-bit ERGBFormat::RGBA and ERGBFormat::Gray were used to describe float pixel formats.
+		// ERGBFormat::RGBA and ERGBFormat::Gray should now only be used for integer channels, while EXR format doesn't support 16-bit integer channels.
+		const TCHAR* FormatName = (InFormat == ERGBFormat::RGBA) ? TEXT("RGBA") : TEXT("Gray");
+		ErrorMessage = FString::Printf(TEXT("Usage of 16-bit ERGBFormat::%s raw format for decompressing float EXR channels is deprecated, please use ERGBFormat::%sF instead."), *FormatName, *FormatName);
+	}
+	else if (InBitDepth != 16 && InBitDepth != 32)
+	{
+		ErrorMessage = TEXT("Unsupported bit depth, expected 16 or 32.");
+	}
+	else if (InFormat != ERGBFormat::RGBAF && InFormat != ERGBFormat::GrayF)
+	{
+		// EXR uint32 channels are currently not supported
+		ErrorMessage = TEXT("Unsupported RGB format, expected ERGBFormat::RGBAF or ERGBFormat::GrayF.");
 	}
-}
 
-void FExrImageWrapper::Uncompress( const ERGBFormat InFormat, const int32 InBitDepth )
-{
-	// Ensure we haven't already uncompressed the file.
-	if ( RawData.Num() != 0 )
+	if (!ErrorMessage.IsEmpty())
 	{
+		UE_LOG(LogImageWrapper, Error, TEXT("Cannot decompress EXR image: %s."), *ErrorMessage);
+		SetError(*ErrorMessage);
 		return;
 	}
 
-	FMemFileIn MemFile(&CompressedData[0], CompressedData.Num());
-
-	Imf::RgbaInputFile ImfFile(MemFile);
-
-	Imath::Box2i win = ImfFile.dataWindow();
+	const char* const* ChannelNames;
+	int32 ChannelCount = GetChannelNames(InFormat, ChannelNames);
+	check(ChannelCount);
 
-	check(BitDepth == 16);
-	check(Width);
-	check(Height);
+	TArray<TArray64<uint8>> ChannelData;
+	ChannelData.SetNum(ChannelCount);
 
-	uint32 Channels = 4;
-	
-	RawData.Empty();
-	RawData.AddUninitialized( int64(Width) * int64(Height) * int64(Channels) * int64(BitDepth / 8) );
-
-	int dx = win.min.x;
-	int dy = win.min.y;
-
-	ImfFile.setFrameBuffer((Imf::Rgba*)&RawData[0] - int64(dx) - int64(dy) * int64(Width), 1, Width);
+	Imf::PixelType ImfPixelType = (InBitDepth == 16) ? Imf::HALF : Imf::FLOAT;
+	int32 BytesPerChannelPixel = (ImfPixelType == Imf::HALF) ? 2 : 4;
 
+	// openEXR can throw exceptions when parsing invalid data.
 	try
 	{
-		ImfFile.readPixels(win.min.y, win.max.y);
+		Imf::FrameBuffer ImfFrameBuffer;
+		for (int32 c = 0; c < ChannelCount; ++c)
+		{
+			ChannelData[c].SetNumUninitialized((int64)BytesPerChannelPixel * Width * Height);
+			// Use 1.0 as a default value for the alpha channel, in case if it is not present in the EXR, use 0.0 for all other channels.
+			double DefaultValue = !strcmp(ChannelNames[c], "A") ? 1.0 : 0.0;
+			ImfFrameBuffer.insert(ChannelNames[c], Imf::Slice(ImfPixelType, (char*)ChannelData[c].GetData(), BytesPerChannelPixel, (size_t)BytesPerChannelPixel * Width, 1, 1, DefaultValue));
+		}
+		FMemFileIn MemFile(CompressedData.GetData(), CompressedData.Num());
+		Imf::InputFile ImfFile(MemFile);
+		Imf::Header ImfHeader = ImfFile.header();
+		Imath::Box2i ImfDataWindow = ImfHeader.dataWindow();
+		ImfFile.setFrameBuffer(ImfFrameBuffer);
+		ImfFile.readPixels(ImfDataWindow.min.y, ImfDataWindow.max.y);
 	}
 	catch (const std::exception& Exception)
 	{
 		TStringConversion<TStringConvert<char, TCHAR>> Convertor(Exception.what());
-		UE_LOG(LogImageWrapper, Error, TEXT("Cannot read EXR file: %s"), Convertor.Get());
+		UE_LOG(LogImageWrapper, Error, TEXT("Cannot decompress EXR image: %s"), Convertor.Get());
 		SetError(Convertor.Get());
+		return;
 	}
-}
-
-// from http://www.openexr.com/ReadingAndWritingImageFiles.pdf
-bool IsThisAnOpenExrFile(Imf::IStream& f)
-{
-	char b[4];
-	f.read(b, sizeof(b));
-
-	f.seekg(0);
-
-	return b[0] == 0x76 && b[1] == 0x2f && b[2] == 0x31 && b[3] == 0x01;
-}
-
-bool FExrImageWrapper::SetCompressed( const void* InCompressedData, int64 InCompressedSize )
-{
-	if(!FImageWrapperBase::SetCompressed( InCompressedData, InCompressedSize))
-	{
-		return false;
-	}
-
-	FMemFileIn MemFile(InCompressedData, InCompressedSize);
 
-	if(!IsThisAnOpenExrFile(MemFile))
+	// EXR channels are compressed non-interleaved.
+	RawData.SetNumUninitialized((int64)BytesPerChannelPixel * ChannelCount * Width * Height);
+	for (int64 OffsetNonInterleaved = 0, OffsetInterleaved = 0; OffsetInterleaved < RawData.Num(); OffsetNonInterleaved += BytesPerChannelPixel)
 	{
-		return false;
+		for (int32 c = 0; c < ChannelCount; ++c)
+		{
+			for (int32 b = 0; b < BytesPerChannelPixel; ++b, ++OffsetInterleaved)
+			{
+				RawData[OffsetInterleaved] = ChannelData[c][OffsetNonInterleaved + b];
+			}
+		}
 	}
 
-	Imf::RgbaInputFile ImfFile(MemFile);
-
-	Imath::Box2i win = ImfFile.dataWindow();
-
-	Imath::V2i dim(win.max.x - win.min.x + 1, win.max.y - win.min.y + 1);
-
-	BitDepth = 16;
+	RawFormat = InFormat;
+	RawBitDepth = InBitDepth;
+}
 
-	Width = dim.x;
-	Height = dim.y;
 
-	Format = ERGBFormat::RGBAF;
 
-	return true;
-}
 
 #endif // WITH_UNREALEXR
diff --git a/Engine/Source/Runtime/ImageWrapper/Private/Formats/ExrImageWrapper.h b/Engine/Source/Runtime/ImageWrapper/Private/Formats/ExrImageWrapper.h
--- a/Engine/Source/Runtime/ImageWrapper/Private/Formats/ExrImageWrapper.h
+++ b/Engine/Source/Runtime/ImageWrapper/Private/Formats/ExrImageWrapper.h
@@ -39,23 +39,10 @@ public:
 
 	//~ FImageWrapper interface
 
+	virtual bool SetRaw(const void* InRawData, int64 InRawSize, const int32 InWidth, const int32 InHeight, const ERGBFormat InFormat, const int32 InBitDepth) override;
+	virtual bool SetCompressed(const void* InCompressedData, int64 InCompressedSize) override;
 	virtual void Compress(int32 Quality) override;
 	virtual void Uncompress(const ERGBFormat InFormat, int32 InBitDepth) override;
-	virtual bool SetCompressed(const void* InCompressedData, int64 InCompressedSize) override;
-
-protected:
-
-	template <Imf::PixelType OutputFormat, typename sourcetype>
-	void WriteFrameBufferChannel(Imf::FrameBuffer& ImfFrameBuffer, const char* ChannelName, const sourcetype* SrcData, TArray64<uint8>& ChannelBuffer);
-
-	template <Imf::PixelType OutputFormat, typename sourcetype>
-	void CompressRaw(const sourcetype* SrcData, bool bIgnoreAlpha);
-
-	const char* GetRawChannelName(int ChannelIndex) const;
-
-private:
-
-	bool bUseCompression;
 };
 
 #endif // WITH_UNREALEXR
diff --git a/Engine/Source/Runtime/ImageWrapper/Public/IImageWrapper.h b/Engine/Source/Runtime/ImageWrapper/Public/IImageWrapper.h
--- a/Engine/Source/Runtime/ImageWrapper/Public/IImageWrapper.h
+++ b/Engine/Source/Runtime/ImageWrapper/Public/IImageWrapper.h
@@ -64,7 +64,10 @@ enum class ERGBFormat : int8
 	RGBAF = 3,
 
 	// Blue, Green, Red and Exponent (Similar to the RGBE format from radiance but with the blue and red channel inversed)
-	BGRE = 4,
+	BGRE =  4,
+
+	// Gray scale using IEEE Floating-Point Arithmetic (see IEEE754). The format is always binary.
+	GrayF = 5,
 };
 
 
