Some minor clean ups for the WITH_UNREALPNG/WITH_UNREALJPEG code paths in ImageWrapper.

#rnx

- When libjpeg-turbo is enabled we now disable libJPG, before we were including both and relying on libJPG being stripped out.
- Fixed IcoImageWrapper.cpp to compile if WITH_UNREALPNG is 0.
- Cleaned up code comments around enabling libjpeg-turbo with a reminder to investigate the exe size increase if anyone tries to enable it on a console platform.

diff --git a/Engine/Source/Runtime/ImageWrapper/ImageWrapper.Build.cs b/Engine/Source/Runtime/ImageWrapper/ImageWrapper.Build.cs
--- a/Engine/Source/Runtime/ImageWrapper/ImageWrapper.Build.cs
+++ b/Engine/Source/Runtime/ImageWrapper/ImageWrapper.Build.cs
@@ -19,15 +19,17 @@ public class ImageWrapper : ModuleRules
 
 		AddEngineThirdPartyPrivateStaticDependencies(Target,
 			"zlib",
-			"UElibPNG",
-			"UElibJPG"
+			"UElibPNG"
 		);
 
-		// Add LibJpegTurbo for supported platforms
-		// **** NOTE - Only Win64/Linux has been tested - other platforms are usable at your own risk, but have not been tested
-		if ((Target.Platform == UnrealTargetPlatform.Win64) ||
-			(Target.IsInPlatformGroup(UnrealPlatformGroup.Unix)))
-		/* (Target.Platform == UnrealTargetPlatform.Mac)) */
+		// Jpeg Decoding
+		// LibJpegTurbo is much faster than UElibJPG but has not been compiled or tested for all platforms
+		// Note that currently this module is included at runtime, so consider the increase in exe size before
+		// enabling for any of the console/phone platforms!
+
+		// **** NOTE - Only Win64/Linux has been tested - UnrealTargetPlatform.Mac is available but has not been tested, use at your own risk!
+		if (Target.Platform == UnrealTargetPlatform.Win64 /*|| Target.Platform == UnrealTargetPlatform.Mac*/
+			|| Target.IsInPlatformGroup(UnrealPlatformGroup.Unix))
 		{
 			PublicDefinitions.Add("WITH_LIBJPEGTURBO=1");
 			AddEngineThirdPartyPrivateStaticDependencies(Target, "LibJpegTurbo");
@@ -35,6 +37,7 @@ public class ImageWrapper : ModuleRules
 		else
 		{
 			PublicDefinitions.Add("WITH_LIBJPEGTURBO=0");
+			AddEngineThirdPartyPrivateStaticDependencies(Target, "UElibJPG");
 		}
 
 		// Add openEXR lib for windows builds.
diff --git a/Engine/Source/Runtime/ImageWrapper/Private/Formats/IcoImageWrapper.cpp b/Engine/Source/Runtime/ImageWrapper/Private/Formats/IcoImageWrapper.cpp
--- a/Engine/Source/Runtime/ImageWrapper/Private/Formats/IcoImageWrapper.cpp
+++ b/Engine/Source/Runtime/ImageWrapper/Private/Formats/IcoImageWrapper.cpp
@@ -163,11 +163,13 @@ bool FIcoImageWrapper::LoadICOHeader()
 
 	if (bFoundImage)
 	{
+#if WITH_UNREALPNG
 		if (bIsPng)
 		{
 			SubImageWrapper = PngWrapper;
 		}
 		else
+#endif //WITH_UNREALPNG
 		{
 			SubImageWrapper = BmpWrapper;
 		}
diff --git a/Engine/Source/Runtime/ImageWrapper/Private/Formats/JpegImageWrapper.cpp b/Engine/Source/Runtime/ImageWrapper/Private/Formats/JpegImageWrapper.cpp
--- a/Engine/Source/Runtime/ImageWrapper/Private/Formats/JpegImageWrapper.cpp
+++ b/Engine/Source/Runtime/ImageWrapper/Private/Formats/JpegImageWrapper.cpp
@@ -16,20 +16,21 @@
 	#endif
 #endif
 
-PRAGMA_DISABLE_SHADOW_VARIABLE_WARNINGS
-#include "jpgd.h"
-#include "jpgd.cpp"
-#include "jpge.h"
-#include "jpge.cpp"
-PRAGMA_ENABLE_SHADOW_VARIABLE_WARNINGS
 
 #if WITH_LIBJPEGTURBO
-THIRD_PARTY_INCLUDES_START
-#pragma push_macro("DLLEXPORT")
-#undef DLLEXPORT // libjpeg-turbo defines DLLEXPORT as well
-#include "turbojpeg.h"
-#pragma pop_macro("DLLEXPORT")
-THIRD_PARTY_INCLUDES_END
+	THIRD_PARTY_INCLUDES_START
+	#pragma push_macro("DLLEXPORT")
+	#undef DLLEXPORT // libjpeg-turbo defines DLLEXPORT as well
+	#include "turbojpeg.h"
+	#pragma pop_macro("DLLEXPORT")
+	THIRD_PARTY_INCLUDES_END
+#else
+	PRAGMA_DISABLE_SHADOW_VARIABLE_WARNINGS
+	#include "jpgd.h"
+	#include "jpgd.cpp"
+	#include "jpge.h"
+	#include "jpge.cpp"
+	PRAGMA_ENABLE_SHADOW_VARIABLE_WARNINGS
 #endif	// WITH_LIBJPEGTURBO
 
 #ifdef __clang__
