Fixed an undefined behavior caused by a stack variable being use after the end of its scope.

diff --git a/Engine/Source/Runtime/ImageWrapper/Private/Formats/TiffImageWrapper.cpp b/Engine/Source/Runtime/ImageWrapper/Private/Formats/TiffImageWrapper.cpp
--- a/Engine/Source/Runtime/ImageWrapper/Private/Formats/TiffImageWrapper.cpp
+++ b/Engine/Source/Runtime/ImageWrapper/Private/Formats/TiffImageWrapper.cpp
@@ -29,20 +29,20 @@ namespace UE::ImageWrapper::Private
 	namespace TiffImageWrapper
 	{
 		template<class DataTypeDest, class DataTypeSrc>
-		DataTypeDest ConvertToWriteFormat(DataTypeSrc ReadedValue)
+		DataTypeDest ConvertToWriteFormat(DataTypeSrc ReadValue)
 		{
 			if constexpr (TIsSame<DataTypeDest, DataTypeSrc>::Value || TIsSame<DataTypeDest, FFloat16>::Value)
 			{
-				return ReadedValue;
+				return ReadValue;
 			}
 			else if constexpr (TIsSame<DataTypeSrc, FFloat16>::Value || TIsFloatingPoint<DataTypeSrc>::Value)
 			{
-				return TNumericLimits<DataTypeDest>::Max() * FMath::Clamp(ReadedValue, 0.0, 1.0);
+				return TNumericLimits<DataTypeDest>::Max() * FMath::Clamp(ReadValue, 0.0, 1.0);
 			}
 			else
 			{
 				// Naive linear interpolation
-				return TNumericLimits<DataTypeDest>::Max() * (double(ReadedValue) / double(TNumericLimits<DataTypeSrc>::Max()));
+				return TNumericLimits<DataTypeDest>::Max() * (double(ReadValue) / double(TNumericLimits<DataTypeSrc>::Max()));
 			}
 		}
 
@@ -128,8 +128,8 @@ namespace UE::ImageWrapper::Private
 			{
 				constexpr uint8 MaxValue = uint8(0xff) >> ((Parent::ChannelPerReadIndex - 1) * NumBits);
 
-				const uint8 ReadedValue = Parent::Read(ReadArray, ReadIndex);
-				WriteArray[WriteIndex] = TNumericLimits<uint8>::Max() * (double(ReadedValue) / MaxValue);
+				const uint8 ReadValue = Parent::Read(ReadArray, ReadIndex);
+				WriteArray[WriteIndex] = TNumericLimits<uint8>::Max() * (double(ReadValue) / MaxValue);
 			}
 		};
 
@@ -156,11 +156,11 @@ namespace UE::ImageWrapper::Private
 				Blues = TArrayView64<uint16>(BluesPtr, NumValues);
 			}
 
-			void Write(TArrayView64<uint16>& WriteArray, int64 WriteIndex, DataTypeSrc ReadedValue)
+			void Write(TArrayView64<uint16>& WriteArray, int64 WriteIndex, DataTypeSrc ReadValue)
 			{
-				WriteArray[WriteIndex] = Reds[ReadedValue];
-				WriteArray[WriteIndex + 1] = Greens[ReadedValue];
-				WriteArray[WriteIndex + 2] = Blues[ReadedValue];
+				WriteArray[WriteIndex] = Reds[ReadValue];
+				WriteArray[WriteIndex + 1] = Greens[ReadValue];
+				WriteArray[WriteIndex + 2] = Blues[ReadValue];
 			}
 
 			TArrayView64<uint16> Reds;
@@ -325,11 +325,11 @@ namespace UE::ImageWrapper::Private
 			int64 RemaningBytes = TiffImageWrapper->CompressedData.Num() - TiffImageWrapper->CurrentPosition;
 			if (RemaningBytes > 0)
 			{
-				int64 NumBytesReaded = FMath::Min<int64>(RemaningBytes, Size);
-				FMemory::Memcpy(Buffer, TiffImageWrapper->CompressedData.GetData() + TiffImageWrapper->CurrentPosition, NumBytesReaded);
-				TiffImageWrapper->CurrentPosition += NumBytesReaded;
+				int64 NumBytesRead = FMath::Min<int64>(RemaningBytes, Size);
+				FMemory::Memcpy(Buffer, TiffImageWrapper->CompressedData.GetData() + TiffImageWrapper->CurrentPosition, NumBytesRead);
+				TiffImageWrapper->CurrentPosition += NumBytesRead;
 
-				return NumBytesReaded;
+				return NumBytesRead;
 			}
 			return 0;
 		}
@@ -849,11 +849,11 @@ namespace UE::ImageWrapper::Private
 
 		ReadWriteAdapter Adapter(Tiff);
 
-		auto ProccessDecodedData =
+		auto ProcessDecodedData =
 				[this, &WriteArray, NumOfChannelDest, StepDest, NumberOfChannelInReadArray, TileWidth, TileHeight, &Adapter]
-				(int32 NumberOfColumnReaded, int32 NumberOfRowReaded, TArray64<uint8>& ReadedBuffer, uint8 SampleIndex, int32 BlockX, int32 BlockY)
+				(int32 NumberOfColumnRead, int32 NumberOfRowRead, TArray64<uint8>& ReadBuffer, uint8 SampleIndex, int32 BlockX, int32 BlockY)
 				{
-					TArrayView64<const DataTypeSrc> ReadArray(static_cast<DataTypeSrc*>(static_cast<void*>(ReadedBuffer.GetData())), ReadedBuffer.Num() / sizeof(DataTypeSrc));
+					TArrayView64<const DataTypeSrc> ReadArray(static_cast<DataTypeSrc*>(static_cast<void*>(ReadBuffer.GetData())), ReadBuffer.Num() / sizeof(DataTypeSrc));
 					uint64 CurrentOffset = int64(BlockY) * Width * NumOfChannelDest;
 					if constexpr (bIsTiled)
 					{
@@ -861,16 +861,16 @@ namespace UE::ImageWrapper::Private
 						CurrentOffset += int64(BlockX) * TileWidth * NumOfChannelDest;
 					}
 	
-					ParallelFor(NumberOfRowReaded * NumberOfColumnReaded
-						, [this, BlockY, &WriteArray, &ReadArray, NumOfChannelDest, StepDest, NumberOfChannelInReadArray, SampleIndex, CurrentOffset, NumberOfColumnReaded, TileWidth, &Adapter]
+					ParallelFor(NumberOfRowRead * NumberOfColumnRead
+						, [this, BlockY, &WriteArray, &ReadArray, NumOfChannelDest, StepDest, NumberOfChannelInReadArray, SampleIndex, CurrentOffset, NumberOfColumnRead, TileWidth, &Adapter]
 							(int32 PixelReadIndex)
 							{
 								int64 ReadIndex;
 								int64 WriteIndex;
 								if constexpr (bIsTiled)
 								{
-									const int32 PositionXInTile = PixelReadIndex % NumberOfColumnReaded;
-									const int32 PositionYInTile = PixelReadIndex / NumberOfColumnReaded;
+									const int32 PositionXInTile = PixelReadIndex % NumberOfColumnRead;
+									const int32 PositionYInTile = PixelReadIndex / NumberOfColumnRead;
 									WriteIndex = CurrentOffset + PositionXInTile * NumOfChannelDest + PositionYInTile * Width * NumOfChannelDest + SampleIndex;
 
 									// The end of tile in X can be some garbage that act as padding data so that all tiles are the same size
@@ -933,9 +933,9 @@ namespace UE::ImageWrapper::Private
 						}
 
 						const int32 NumberOfColumn = X + TileWidth > Width ? Width - X : TileWidth;
-						Tasks.Add(TGraphTask<FProcessDecodedDataTask>::CreateTask().ConstructAndDispatchWhenReady([&UsableBufferQueue, &ProccessDecodedData, NumberOfColumn, NumberOfRow, BufferPtr, SampleIndex, TileX, TileY]()
+						Tasks.Add(TGraphTask<FProcessDecodedDataTask>::CreateTask().ConstructAndDispatchWhenReady([&UsableBufferQueue, &ProcessDecodedData, NumberOfColumn, NumberOfRow, BufferPtr, SampleIndex, TileX, TileY]()
 							{
-								ProccessDecodedData(NumberOfColumn, NumberOfRow, *(BufferPtr.Get()), SampleIndex, TileX, TileY);
+								ProcessDecodedData(NumberOfColumn, NumberOfRow, *(BufferPtr.Get()), SampleIndex, TileX, TileY);
 								UsableBufferQueue.Enqueue(BufferPtr);
 							}));
 					}
@@ -957,11 +957,6 @@ namespace UE::ImageWrapper::Private
 
 			const TCHAR* ErrorMessage = TEXT("Couldn't open Tiff strip. This is generally caused by a compression format that UE don't support.");
 
-			auto ProcessDecodedStrips = [this, &ProccessDecodedData](int32 NumberOfRowReaded, TArray64<uint8>& ReadedBuffer, int32 FirstLineReaded, uint8 SampleIndex)
-				{
-					ProccessDecodedData(Width, NumberOfRowReaded, ReadedBuffer, SampleIndex, 0, FirstLineReaded);
-				};
-
 			Tasks.Reserve(TIFFNumberOfStrips(Tiff));
 
 			for (uint8 SampleIndex = 0; SampleIndex < NumOfPlanes; SampleIndex++)
@@ -978,17 +973,18 @@ namespace UE::ImageWrapper::Private
 						}
 
 						// The last strip might be smaller
-						int32 NumberOfRowReaded = Y + RowPerStrip > Height ? Height - Y : RowPerStrip;
-						if (TIFFReadEncodedStrip(Tiff, TIFFComputeStrip(Tiff, Y, SampleIndex), BufferPtr->GetData(), NumberOfRowReaded * LineSizeScr) == -1)
+						int32 NumberOfRowRead = Y + RowPerStrip > Height ? Height - Y : RowPerStrip;
+						if (TIFFReadEncodedStrip(Tiff, TIFFComputeStrip(Tiff, Y, SampleIndex), BufferPtr->GetData(), NumberOfRowRead * LineSizeScr) == -1)
 						{
 							SetError(ErrorMessage);
 							return false;
 						}
 
 						
-						Tasks.Add(TGraphTask<FProcessDecodedDataTask>::CreateTask().ConstructAndDispatchWhenReady([&UsableBufferQueue, &ProcessDecodedStrips, NumberOfRowReaded, BufferPtr, Y, SampleIndex]()
+						Tasks.Add(TGraphTask<FProcessDecodedDataTask>::CreateTask().ConstructAndDispatchWhenReady([this, &ProcessDecodedData, &UsableBufferQueue, NumberOfRowRead, BufferPtr, Y, SampleIndex]()
 							{
-								ProcessDecodedStrips(NumberOfRowReaded, *(BufferPtr.Get()), Y, SampleIndex);
+								// Consider the strip to be a tile that as the full with of the image.
+								ProcessDecodedData(Width, NumberOfRowRead, *(BufferPtr.Get()), SampleIndex, 0, Y);
 								UsableBufferQueue.Enqueue(BufferPtr);
 							}));
 					}
@@ -1008,9 +1004,10 @@ namespace UE::ImageWrapper::Private
 						return false;
 					}
 
-					Tasks.Add(TGraphTask<FProcessDecodedDataTask>::CreateTask().ConstructAndDispatchWhenReady([this, &UsableBufferQueue, &ProcessDecodedStrips, BufferPtr, SampleIndex]()
+					Tasks.Add(TGraphTask<FProcessDecodedDataTask>::CreateTask().ConstructAndDispatchWhenReady([this, &UsableBufferQueue, &ProcessDecodedData, BufferPtr, SampleIndex]()
 						{
-							ProcessDecodedStrips(Height, *(BufferPtr.Get()), 0, SampleIndex);
+							// Consider the image to be one big tile.
+							ProcessDecodedData(Width, Height, *(BufferPtr.Get()), SampleIndex, 0, 0);
 							UsableBufferQueue.Enqueue(BufferPtr);
 						}));
 				}
