ImageWrapper: remove unused animation frame functions
were marked deprecated in previous commit

diff --git a/Engine/Source/Runtime/ImageWrapper/Private/Formats/BmpImageWrapper.cpp b/Engine/Source/Runtime/ImageWrapper/Private/Formats/BmpImageWrapper.cpp
--- a/Engine/Source/Runtime/ImageWrapper/Private/Formats/BmpImageWrapper.cpp
+++ b/Engine/Source/Runtime/ImageWrapper/Private/Formats/BmpImageWrapper.cpp
@@ -517,7 +517,6 @@ void FBmpImageWrapper::Compress(int32 Quality)
 	
 	int OutputRowPadBytes = OutputRowBytes - Width*OutputBytesPerPel;
 	check( OutputRowPadBytes < 4 );
-	uint8 Zero[4] = { };
 
 	// write rows :
 	switch(OutputBytesPerPel)
@@ -529,7 +528,7 @@ void FBmpImageWrapper::Compress(int32 Quality)
 		{
 			memcpy(PayloadPtr,RawPtr + y * Width,Width);
 			PayloadPtr += Width;
-			memcpy(PayloadPtr,Zero,OutputRowPadBytes);
+			memset(PayloadPtr,0,OutputRowPadBytes);
 			PayloadPtr += OutputRowPadBytes;
 		}
 		break;
@@ -547,7 +546,7 @@ void FBmpImageWrapper::Compress(int32 Quality)
 				*PayloadPtr++ = RawRow[x].G;
 				*PayloadPtr++ = RawRow[x].R;
 			}
-			memcpy(PayloadPtr,Zero,OutputRowPadBytes);
+			memset(PayloadPtr,0,OutputRowPadBytes);
 			PayloadPtr += OutputRowPadBytes;
 		}
 		break;
diff --git a/Engine/Source/Runtime/ImageWrapper/Private/Formats/HdrImageWrapper.cpp b/Engine/Source/Runtime/ImageWrapper/Private/Formats/HdrImageWrapper.cpp
--- a/Engine/Source/Runtime/ImageWrapper/Private/Formats/HdrImageWrapper.cpp
+++ b/Engine/Source/Runtime/ImageWrapper/Private/Formats/HdrImageWrapper.cpp
@@ -211,22 +211,6 @@ ERGBFormat FHdrImageWrapper::GetFormat() const
 	return ERGBFormat::BGRE;
 }
 
-bool FHdrImageWrapper::SetAnimationInfo_DEPRECATED(int32 InNumFrames, int32 InFramerate)
-{
-	unimplemented();
-	return false;
-}
-
-int32 FHdrImageWrapper::GetNumFrames_DEPRECATED() const
-{
-	return INDEX_NONE;
-}
-
-int32 FHdrImageWrapper::GetFramerate_DEPRECATED() const
-{
-	return INDEX_NONE;
-}
-
 const FText& FHdrImageWrapper::GetErrorMessage() const
 {
 	return ErrorMessage;
diff --git a/Engine/Source/Runtime/ImageWrapper/Private/ImageWrapperBase.cpp b/Engine/Source/Runtime/ImageWrapper/Private/ImageWrapperBase.cpp
--- a/Engine/Source/Runtime/ImageWrapper/Private/ImageWrapperBase.cpp
+++ b/Engine/Source/Runtime/ImageWrapper/Private/ImageWrapperBase.cpp
@@ -15,8 +15,6 @@ FImageWrapperBase::FImageWrapperBase()
 	, BitDepth(0)
 	, Width(0)
 	, Height(0)
-	, NumFrames_DEPRECATED(1)
-	, Framerate_DEPRECATED(0)
 { }
 
 
@@ -38,8 +36,6 @@ void FImageWrapperBase::Reset()
 	BitDepth = 0;
 	Width = 0;
 	Height = 0;
-	NumFrames_DEPRECATED = 1;
-	Framerate_DEPRECATED = 0;
 }
 
 
@@ -128,13 +124,6 @@ bool FImageWrapperBase::SetRaw(const void* InRawData, int64 InRawSize, const int
 	return true;
 }
 
-bool FImageWrapperBase::SetAnimationInfo_DEPRECATED(int32 InNumFrames, int32 InFramerate)
-{
-	NumFrames_DEPRECATED = InNumFrames;
-	Framerate_DEPRECATED = InFramerate;
-	
-	return true;
-}
 
 
 ERawImageFormat::Type IImageWrapper::ConvertRGBFormat(ERGBFormat RGBFormat,int BitDepth,bool * bIsExactMatch)
diff --git a/Engine/Source/Runtime/ImageWrapper/Private/ImageWrapperBase.h b/Engine/Source/Runtime/ImageWrapper/Private/ImageWrapperBase.h
--- a/Engine/Source/Runtime/ImageWrapper/Private/ImageWrapperBase.h
+++ b/Engine/Source/Runtime/ImageWrapper/Private/ImageWrapperBase.h
@@ -102,19 +102,8 @@ public:
 		return Width;
 	}
 
-	virtual int32 GetNumFrames_DEPRECATED() const override
-	{
-		return NumFrames_DEPRECATED;
-	}
-	
-	virtual int32 GetFramerate_DEPRECATED() const override
-	{
-		return Framerate_DEPRECATED;
-	}
-
 	virtual bool SetCompressed(const void* InCompressedData, int64 InCompressedSize) override;
 	virtual bool SetRaw(const void* InRawData, int64 InRawSize, const int32 InWidth, const int32 InHeight, const ERGBFormat InFormat, const int32 InBitDepth, const int32 InBytesPerRow = 0) override;
-	virtual bool SetAnimationInfo_DEPRECATED(int32 InNumFrames, int32 InFramerate) override;
 
 protected:
 
@@ -139,10 +128,6 @@ protected:
 	int32 Width;
 	int32 Height;
 	
-	/** Animation information */
-	int32 NumFrames_DEPRECATED;
-	int32 Framerate_DEPRECATED;
-
 	/** Last Error Message. */
 	FString LastError;
 };
diff --git a/Engine/Source/Runtime/ImageWrapper/Public/Formats/HdrImageWrapper.h b/Engine/Source/Runtime/ImageWrapper/Public/Formats/HdrImageWrapper.h
--- a/Engine/Source/Runtime/ImageWrapper/Public/Formats/HdrImageWrapper.h
+++ b/Engine/Source/Runtime/ImageWrapper/Public/Formats/HdrImageWrapper.h
@@ -31,12 +31,9 @@ public:
 	virtual int32 GetBitDepth() const override;
 	virtual ERGBFormat GetFormat() const override;
 
-	virtual bool SetAnimationInfo_DEPRECATED(int32 InNumFrames, int32 InFramerate) override;
-	virtual int32 GetNumFrames_DEPRECATED() const override;
-	virtual int32 GetFramerate_DEPRECATED() const override;
 	// IImageWrapper Interface end
 
-	// GetErrorMessage : nice idea, but not virtual, never called
+	// GetErrorMessage : nice idea, but not virtual, never called by the standard import path
 	const FText& GetErrorMessage() const;
 
 	void FreeCompressedData();
diff --git a/Engine/Source/Runtime/ImageWrapper/Public/IImageWrapper.h b/Engine/Source/Runtime/ImageWrapper/Public/IImageWrapper.h
--- a/Engine/Source/Runtime/ImageWrapper/Public/IImageWrapper.h
+++ b/Engine/Source/Runtime/ImageWrapper/Public/IImageWrapper.h
@@ -315,27 +315,6 @@ public:
 		return Ret;
 	}
 
-
-	// @@!! is this animation stuff used at all? -> it seems not
-
-	/**
-	 * Set information for animated formats
-	 * @param InNumFrames The number of frames in the animation (the RawData from SetRaw will need to be a multiple of NumFrames)
-	 * @param InFramerate The playback rate of the animation
-	 * @return true if successful
-	 */
-	virtual bool SetAnimationInfo_DEPRECATED(int32 InNumFrames, int32 InFramerate) = 0;
-
-	/**
-	 * @return The number of frames in an animated image
-	 */
-	virtual int32 GetNumFrames_DEPRECATED() const = 0;
-
-	/**
-	 * @return The playback framerate of animated images (or 0 for non-animated)
-	 */
-	virtual int32 GetFramerate_DEPRECATED() const = 0;
-
 public:
 
 	/** Virtual destructor. */
