DDS image wrapper
allows load/save of only 2d DDS images through the generic imagewrappers apis
can be used with screenshot to dds output for example

diff --git a/Engine/Source/Runtime/ImageWrapper/Private/Formats/DdsImageWrapper.cpp b/Engine/Source/Runtime/ImageWrapper/Private/Formats/DdsImageWrapper.cpp
new file mode 100644
--- /dev/null
+++ b/Engine/Source/Runtime/ImageWrapper/Private/Formats/DdsImageWrapper.cpp
@@ -0,0 +1,208 @@
+#include "DdsImageWrapper.h"
+#include "ImageWrapperPrivate.h"
+
+void FDdsImageWrapper::Reset()
+{
+	Super::Reset();
+
+	FreeDDS();
+}
+
+void FDdsImageWrapper::Compress(int32 Quality)
+{
+	CompressedData.Reset();
+	FreeDDS();
+
+	bool bIsExactMatch;
+	ERawImageFormat::Type RawImageFormat = GetClosestRawImageFormat(&bIsExactMatch);
+	if ( RawImageFormat == ERawImageFormat::Invalid )
+	{
+		// should not get here if caller checked CanSetRawFormat() like they're supposed to
+		SetError(TEXT("Format not supported"));
+		check( CompressedData.IsEmpty() ); // signals error
+		return;
+	}
+
+	// @@!! SRGB
+	EGammaSpace GammaSpace = ERawImageFormat::GetDefaultGammaSpace(RawImageFormat);
+	
+	// some code dupe with IImageWrapper::GetRawImage
+	// @todo Oodle : refactor so this can be shared
+	//   after someone does SetRaw() , I should be able to get an FImage view of Raw bits
+	//	 for my own writers to use
+	//	(can't just use GetRawImage because that's like a GetRaw after SetCompressed)
+
+	FImageView RawImage(RawData.GetData(),Width,Height,1,RawImageFormat,GammaSpace);
+	FImage TempImage;
+
+	if ( ! bIsExactMatch )
+	{
+		// handle non-mapped cases :
+	
+		switch(RawFormat)
+		{
+		case ERGBFormat::RGBA:
+		{
+			// RGBA8 -> BGRA8
+			check( BitDepth == 8 );
+			check( RawImageFormat == ERawImageFormat::BGRA8 );
+			FImageCore::CopyImageRGBABGRA(RawImage, RawImage );
+			break;
+		}
+			
+		case ERGBFormat::BGRA:
+		{
+			// BGRA16 -> RGBA16
+			check( BitDepth == 16 );
+			check( RawImageFormat == ERawImageFormat::RGBA16 );
+			FImageCore::CopyImageRGBABGRA(RawImage, RawImage );
+			break;
+		}
+
+		case ERGBFormat::GrayF:
+		{
+			// 1 channel F32 -> 4xF32 :
+			check( BitDepth == 32 );
+			check( RawImageFormat == ERawImageFormat::RGBA32F );
+			const float * Src = (const float *)RawImage.RawData;
+			FLinearColor * Dst = (FLinearColor *)TempImage.RawData.GetData();
+			int64 NumPixels = RawImage.GetNumPixels();
+			for(int64 i=0;i<NumPixels;i++)
+			{
+				Dst[i] = FLinearColor( Src[i],Src[i],Src[i],1.f );
+			}
+			RawImage = TempImage;
+			break;
+		}
+
+		default:
+			check(0);
+			break;
+		}
+	}
+
+	// RawImage is ready to write to a DDS
+	
+	UE::DDS::EDXGIFormat DXGIFormat = UE::DDS::DXGIFormatFromRawFormat(RawImage.Format,RawImage.GammaSpace);
+
+	DDS = UE::DDS::FDDSFile::CreateEmpty2D(Width,Height,1,DXGIFormat,0);
+
+	DDS->FillMip(RawImage,0);
+	
+	UE::DDS::EDDSError Error = DDS->WriteDDS(CompressedData);
+	
+	FreeDDS();
+
+	if ( Error != UE::DDS::EDDSError::OK )
+	{
+		SetError(TEXT("WriteDDS failed"));
+		CompressedData.Empty(); // signals error
+		return;
+	}
+
+	return;
+}
+
+bool FDdsImageWrapper::SetCompressed(const void* InCompressedData, int64 InCompressedSize)
+{
+	Super::SetCompressed(InCompressedData,InCompressedSize);
+	FreeDDS();
+
+	UE::DDS::EDDSError Error;
+	DDS = UE::DDS::FDDSFile::CreateFromDDSInMemory((const uint8 *)InCompressedData,InCompressedSize,&Error);
+	if ( DDS == nullptr || Error != UE::DDS::EDDSError::OK )
+	{
+		SetError(TEXT("CreateFromDDSInMemory failed"));
+		RawData.Empty();
+		FreeDDS();
+		return false;
+	}
+
+	// populate header info :
+	
+	// change X8 formats to A8 :	
+	DDS->ConvertRGBXtoRGBA();
+	// change RGBA8 to BGRA8 before DXGIFormatGetClosestRawFormat :
+	DDS->ConvertChannelOrder(UE::DDS::EChannelOrder::BGRA);
+
+	// map format to RawFormat and ETextureSourceFormat
+	ERawImageFormat::Type RawImageFormat = UE::DDS::DXGIFormatGetClosestRawFormat(DDS->DXGIFormat);
+	if ( RawImageFormat == ERawImageFormat::Invalid )
+	{
+		//Warn->Logf(ELogVerbosity::Error, TEXT("DDS DXGIFormat not supported : %d : %s"), (int)DDS->DXGIFormat, UE::DDS::DXGIFormatGetName(DDS->DXGIFormat) );
+		
+		SetError(TEXT("DDS DXGIFormat not supported"));
+		RawData.Empty();
+		FreeDDS();
+		return false;
+	}
+
+	if ( ! DDS->IsValidTexture2D() )
+	{
+		UE_LOG(LogImageWrapper, Warning, TEXT("DDS is a complex 3d/cube/array image, DdsImageWrapper will just import the first 2d surface"));
+	}
+	
+	int OutBitDepth;
+	ConvertRawImageFormat(RawImageFormat, RawFormat,OutBitDepth);
+	Format = RawFormat;
+	BitDepth = RawBitDepth = OutBitDepth;
+
+	Width = DDS->Width;
+	Height = DDS->Height;
+
+	RawBytesPerRow = Width * ERawImageFormat::GetBytesPerPixel(RawImageFormat);
+	
+	return true;
+}
+
+void FDdsImageWrapper::Uncompress(const ERGBFormat InFormat, int32 InBitDepth)
+{
+	RawData.Reset();
+	// SetCompressed made the DDS
+	check( DDS != nullptr );
+	// new style, only support decoding to own format :
+	check( InFormat == GetFormat() );
+	check( InBitDepth == GetBitDepth() );
+		
+	ERawImageFormat::Type RawImageFormat = UE::DDS::DXGIFormatGetClosestRawFormat(DDS->DXGIFormat);
+
+	// @@!! SRGB
+	FImage RawImage(Width,Height,RawImageFormat);
+
+	if ( ! DDS->GetMipImage(RawImage,0) )
+	{
+		SetError(TEXT("DDS GetMipImage failed"));
+		FreeDDS();		
+		check( RawData.IsEmpty() ); // indicates failure
+		return;
+	}
+
+	RawData = MoveTemp(RawImage.RawData);
+}
+	
+bool FDdsImageWrapper::CanSetRawFormat(const ERGBFormat InFormat, const int32 InBitDepth) const
+{
+	// accept any format that can roundtrip through ERawImageFormat :
+	// ! exact match is okay, that's an RB swap
+	
+	ERawImageFormat::Type RawImageFormat = ConvertRGBFormat(InFormat,InBitDepth,nullptr);
+	
+	ERGBFormat OutFormat;
+	int OutBitDepth;
+	ConvertRawImageFormat(RawImageFormat, OutFormat,OutBitDepth);
+
+	if ( InFormat == OutFormat && InBitDepth == OutBitDepth )
+	{
+		return true;
+	}
+	else
+	{
+		return false;
+	}
+}
+
+ERawImageFormat::Type FDdsImageWrapper::GetSupportedRawFormat(const ERawImageFormat::Type InFormat) const
+{
+	// all raw formats supported
+	return InFormat;
+}
diff --git a/Engine/Source/Runtime/ImageWrapper/Private/Formats/DdsImageWrapper.h b/Engine/Source/Runtime/ImageWrapper/Private/Formats/DdsImageWrapper.h
new file mode 100644
--- /dev/null
+++ b/Engine/Source/Runtime/ImageWrapper/Private/Formats/DdsImageWrapper.h
@@ -0,0 +1,49 @@
+// Copyright Epic Games, Inc. All Rights Reserved.
+
+#pragma once
+
+#include "CoreTypes.h"
+#include "ImageWrapperBase.h"
+#include "DDSFile.h"
+
+/**
+ * DDS implementation of the helper class
+ */
+class FDdsImageWrapper : public FImageWrapperBase
+{
+	typedef FImageWrapperBase Super;
+public:
+
+	FDdsImageWrapper()
+	{
+	}
+	~FDdsImageWrapper()
+	{
+		FreeDDS();
+	}
+	
+	void FreeDDS()
+	{
+		if ( DDS != nullptr )
+		{
+			delete DDS;
+			DDS = nullptr;
+		}
+	}
+
+public:
+
+	//~ FImageWrapper interface
+	
+	virtual void Reset() override;
+
+	virtual void Compress(int32 Quality) override;
+	virtual bool SetCompressed(const void* InCompressedData, int64 InCompressedSize) override;
+	virtual void Uncompress(const ERGBFormat InFormat, int32 InBitDepth) override;
+	
+	virtual bool CanSetRawFormat(const ERGBFormat InFormat, const int32 InBitDepth) const override;
+	virtual ERawImageFormat::Type GetSupportedRawFormat(const ERawImageFormat::Type InFormat) const override;
+
+public:
+	UE::DDS::FDDSFile * DDS = nullptr;
+};
diff --git a/Engine/Source/Runtime/ImageWrapper/Private/ImageWrapperBase.cpp b/Engine/Source/Runtime/ImageWrapper/Private/ImageWrapperBase.cpp
--- a/Engine/Source/Runtime/ImageWrapper/Private/ImageWrapperBase.cpp
+++ b/Engine/Source/Runtime/ImageWrapper/Private/ImageWrapperBase.cpp
@@ -8,11 +8,11 @@
  *****************************************************************************/
 
 FImageWrapperBase::FImageWrapperBase()
-	: RawFormat(ERGBFormat::Invalid)
+	: Format(ERGBFormat::Invalid)
+	, RawFormat(ERGBFormat::Invalid)
+	, BitDepth(0)
 	, RawBitDepth(0)
 	, RawBytesPerRow(0)
-	, Format(ERGBFormat::Invalid)
-	, BitDepth(0)
 	, Width(0)
 	, Height(0)
 { }
@@ -25,6 +25,7 @@ void FImageWrapperBase::Reset()
 {
 	LastError.Empty();
 	
+	//@@!!
 	// not reset ?
 	//TArray64<uint8> RawData;
 	//TArray64<uint8> CompressedData;
@@ -62,25 +63,31 @@ bool FImageWrapperBase::GetRaw(const ERGBFormat InFormat, int32 InBitDepth, TArr
 	LastError.Empty();
 	Uncompress(InFormat, InBitDepth);
 
-	if (LastError.IsEmpty())
+	if ( ! LastError.IsEmpty())
 	{
-		OutRawData = MoveTemp(RawData);
-		if ( OutRawData.IsEmpty() )
-		{
-			return false;
-		}
+		return false;
+	}
+	
+	if ( RawData.IsEmpty() )
+	{
+		return false;
 	}
 
-	return LastError.IsEmpty();
+	OutRawData = MoveTemp(RawData);
+	
+	return true;
 }
 
 
 bool FImageWrapperBase::SetCompressed(const void* InCompressedData, int64 InCompressedSize)
 {
+	Reset();
+	RawData.Empty();			// Invalidates the raw data too
+
 	if(InCompressedSize > 0 && InCompressedData != nullptr)
 	{
-		Reset();
-		RawData.Empty();			// Invalidates the raw data too
+		// this is usually an unnecessary allocation and copy
+		// we should decompress directly from the source buffer
 
 		CompressedData.Empty(InCompressedSize);
 		CompressedData.AddUninitialized(InCompressedSize);
@@ -109,6 +116,9 @@ bool FImageWrapperBase::SetRaw(const void* InRawData, int64 InRawSize, const int
 		UE_LOG(LogImageWrapper, Warning, TEXT("ImageWrapper unsupported format; check CanSetRawFormat; %d x %d"), (int)InFormat,InBitDepth);
 		return false;
 	}
+	
+	// this is usually an unnecessary allocation and copy
+	// we should compress directly from the source buffer
 
 	RawData.Empty(InRawSize);
 	RawData.AddUninitialized(InRawSize);
@@ -118,6 +128,9 @@ bool FImageWrapperBase::SetRaw(const void* InRawData, int64 InRawSize, const int
 	RawBitDepth = InBitDepth;
 	RawBytesPerRow = InBytesPerRow;
 
+	Format = RawFormat;
+	BitDepth = RawBitDepth;
+
 	Width = InWidth;
 	Height = InHeight;
 
@@ -205,7 +218,8 @@ ERawImageFormat::Type IImageWrapper::ConvertRGBFormat(ERGBFormat RGBFormat,int B
 		else if ( BitDepth == 32 )
 		{
 			*bIsExactMatch = false; // no single channel F32
-			return ERawImageFormat::RGBA32F; // promote F32 to 4xF32
+			return ERawImageFormat::RGBA32F; // promote F32 to 4xF32 for no quality loss (alternative is 1xF16)
+			//@todo Oodle: add ERawImageFormat::R32F ?
 		}
 		break;
 
diff --git a/Engine/Source/Runtime/ImageWrapper/Private/ImageWrapperBase.h b/Engine/Source/Runtime/ImageWrapper/Private/ImageWrapperBase.h
--- a/Engine/Source/Runtime/ImageWrapper/Private/ImageWrapperBase.h
+++ b/Engine/Source/Runtime/ImageWrapper/Private/ImageWrapperBase.h
@@ -112,17 +112,14 @@ protected:
 	TArray64<uint8> CompressedData;
 
 	/** Format of the raw data */
+	// @@!! two format vars ugh
+	ERGBFormat Format;
 	ERGBFormat RawFormat;
+	int8 BitDepth;
 	int8 RawBitDepth;
 
 	/** Bytes per row for the raw data */
-	int32 RawBytesPerRow;
-
-	/** Format of the image */
-	ERGBFormat Format;
-
-	/** Bit depth of the image */
-	int8 BitDepth;
+	int32 RawBytesPerRow; // @@!! not supported, remove me
 
 	/** Width/Height of the image data */
 	int32 Width;
diff --git a/Engine/Source/Runtime/ImageWrapper/Private/ImageWrapperModule.cpp b/Engine/Source/Runtime/ImageWrapper/Private/ImageWrapperModule.cpp
--- a/Engine/Source/Runtime/ImageWrapper/Private/ImageWrapperModule.cpp
+++ b/Engine/Source/Runtime/ImageWrapper/Private/ImageWrapperModule.cpp
@@ -14,6 +14,9 @@
 #include "Formats/PngImageWrapper.h"
 #include "Formats/TgaImageWrapper.h"
 #include "Formats/TiffImageWrapper.h"
+#include "Formats/DdsImageWrapper.h"
+
+#include "DDSFile.h"
 
 #include "IImageWrapperModule.h"
 
@@ -126,6 +129,10 @@ public:
 			break;
 #endif // WITH_LIBTIFF
 
+		case EImageFormat::DDS:
+			ImageWrapper = MakeShared<FDdsImageWrapper>();
+			break;
+
 		default:
 			break;
 		}
@@ -177,6 +184,10 @@ public:
 		{
 			Format = EImageFormat::TGA;
 		}
+		else if ( UE::DDS::FDDSFile::IsADDS((const uint8 *)CompressedData, CompressedSize) )
+		{
+			Format = EImageFormat::DDS;
+		}
 
 		return Format;
 	}
@@ -195,6 +206,7 @@ public:
 		case EImageFormat::TGA: return TEXT("tga");
 		case EImageFormat::HDR: return TEXT("hdr");
 		case EImageFormat::TIFF: return TEXT("tiff");
+		case EImageFormat::DDS: return TEXT("dds");
 		case EImageFormat::Invalid:
 		default:
 			check(0);
@@ -247,6 +259,10 @@ public:
 		{
 			return EImageFormat::TIFF;
 		}
+		else if ( FCString::Stricmp(Name,TEXT("dds")) == 0 )
+		{
+			return EImageFormat::DDS;
+		}
 		else
 		{
 			UE_LOG(LogImageWrapper,Warning,TEXT("GetImageFormatFromExtension not found : %s\n"),Name);
diff --git a/Engine/Source/Runtime/ImageWrapper/Public/IImageWrapper.h b/Engine/Source/Runtime/ImageWrapper/Public/IImageWrapper.h
--- a/Engine/Source/Runtime/ImageWrapper/Public/IImageWrapper.h
+++ b/Engine/Source/Runtime/ImageWrapper/Public/IImageWrapper.h
@@ -45,6 +45,9 @@ enum class EImageFormat : int8
 
 	/** Tag Image File Format files */
 	TIFF,
+
+	/** DirectDraw Surface */
+	DDS,
 };
 
 
