fix missing copyright text

diff --git a/Engine/Source/Runtime/ImageWrapper/Private/Formats/DdsImageWrapper.cpp b/Engine/Source/Runtime/ImageWrapper/Private/Formats/DdsImageWrapper.cpp
--- a/Engine/Source/Runtime/ImageWrapper/Private/Formats/DdsImageWrapper.cpp
+++ b/Engine/Source/Runtime/ImageWrapper/Private/Formats/DdsImageWrapper.cpp
@@ -1,3 +1,5 @@
+// Copyright Epic Games, Inc. All Rights Reserved.
+
 #include "DdsImageWrapper.h"
 #include "ImageWrapperPrivate.h"
 
