ImageWrapper clean up duplicate variables
remove unsupported row stride and check it is tight packed as expected

diff --git a/Engine/Source/Runtime/ImageWrapper/Private/Formats/BmpImageWrapper.cpp b/Engine/Source/Runtime/ImageWrapper/Private/Formats/BmpImageWrapper.cpp
--- a/Engine/Source/Runtime/ImageWrapper/Private/Formats/BmpImageWrapper.cpp
+++ b/Engine/Source/Runtime/ImageWrapper/Private/Formats/BmpImageWrapper.cpp
@@ -363,15 +363,15 @@ ERawImageFormat::Type FBmpImageWrapper::GetSupportedRawFormat(const ERawImageFor
 
 void FBmpImageWrapper::Compress(int32 Quality)
 {
-	check( RawFormat == ERGBFormat::BGRA || RawFormat == ERGBFormat::Gray );
-	check( RawBitDepth == 8 );
+	check( Format == ERGBFormat::BGRA || Format == ERGBFormat::Gray );
+	check( BitDepth == 8 );
 
 	// write 8,24, or 32 bit bmp
 
 	int64 NumPixels = Width*Height;
 	int64 RawDataSize = RawData.Num();
 
-	int RawBytesPerPel = (RawFormat == ERGBFormat::BGRA) ? 4 : 1;
+	int RawBytesPerPel = (Format == ERGBFormat::BGRA) ? 4 : 1;
 
 	check( RawDataSize == NumPixels*RawBytesPerPel );
 
diff --git a/Engine/Source/Runtime/ImageWrapper/Private/Formats/DdsImageWrapper.cpp b/Engine/Source/Runtime/ImageWrapper/Private/Formats/DdsImageWrapper.cpp
--- a/Engine/Source/Runtime/ImageWrapper/Private/Formats/DdsImageWrapper.cpp
+++ b/Engine/Source/Runtime/ImageWrapper/Private/Formats/DdsImageWrapper.cpp
@@ -41,7 +41,7 @@ void FDdsImageWrapper::Compress(int32 Quality)
 	{
 		// handle non-mapped cases :
 	
-		switch(RawFormat)
+		switch(Format)
 		{
 		case ERGBFormat::RGBA:
 		{
@@ -127,7 +127,7 @@ bool FDdsImageWrapper::SetCompressed(const void* InCompressedData, int64 InCompr
 	// change RGBA8 to BGRA8 before DXGIFormatGetClosestRawFormat :
 	DDS->ConvertChannelOrder(UE::DDS::EChannelOrder::BGRA);
 
-	// map format to RawFormat and ETextureSourceFormat
+	// map format to RawImageFormat and ETextureSourceFormat
 	ERawImageFormat::Type RawImageFormat = UE::DDS::DXGIFormatGetClosestRawFormat(DDS->DXGIFormat);
 	if ( RawImageFormat == ERawImageFormat::Invalid )
 	{
@@ -144,16 +144,11 @@ bool FDdsImageWrapper::SetCompressed(const void* InCompressedData, int64 InCompr
 		UE_LOG(LogImageWrapper, Warning, TEXT("DDS is a complex 3d/cube/array image, DdsImageWrapper will just import the first 2d surface"));
 	}
 	
-	int OutBitDepth;
-	ConvertRawImageFormat(RawImageFormat, RawFormat,OutBitDepth);
-	Format = RawFormat;
-	BitDepth = RawBitDepth = OutBitDepth;
+	ConvertRawImageFormat(RawImageFormat, Format,BitDepth);
 
 	Width = DDS->Width;
 	Height = DDS->Height;
 
-	RawBytesPerRow = Width * ERawImageFormat::GetBytesPerPixel(RawImageFormat);
-	
 	return true;
 }
 
diff --git a/Engine/Source/Runtime/ImageWrapper/Private/Formats/ExrImageWrapper.cpp b/Engine/Source/Runtime/ImageWrapper/Private/Formats/ExrImageWrapper.cpp
--- a/Engine/Source/Runtime/ImageWrapper/Private/Formats/ExrImageWrapper.cpp
+++ b/Engine/Source/Runtime/ImageWrapper/Private/Formats/ExrImageWrapper.cpp
@@ -324,7 +324,7 @@ void FExrImageWrapper::Compress(int32 Quality)
 	TArray64<uint8> ConvertedRawData;
 	bool bNeedsConversion = false;
 
-	if (RawBitDepth == 8)
+	if (BitDepth == 8)
 	{
 		// uint8 channels are linearly converted into FFloat16 channels.
 		// note: NO GAMMA CORRECTION
@@ -339,14 +339,14 @@ void FExrImageWrapper::Compress(int32 Quality)
 	}
 	else
 	{
-		ImfPixelType = (RawBitDepth == 16) ? Imf::HALF : Imf::FLOAT;
+		ImfPixelType = (BitDepth == 16) ? Imf::HALF : Imf::FLOAT;
 	}
 
 	const TArray64<uint8>& PixelData = bNeedsConversion ? ConvertedRawData : RawData;
 
 	const char* const* ChannelNames;
-	int32 ChannelCount = GetChannelNames(RawFormat, ChannelNames);
-	check((int64)ChannelCount * Width * Height * RawBitDepth == RawData.Num() * 8);
+	int32 ChannelCount = GetChannelNames(Format, ChannelNames);
+	check((int64)ChannelCount * Width * Height * BitDepth == RawData.Num() * 8);
 
 	int32 BytesPerChannelPixel = (ImfPixelType == Imf::HALF) ? 2 : 4;
 	TArray<TArray64<uint8>> ChannelData;
@@ -405,7 +405,7 @@ void FExrImageWrapper::Uncompress(const ERGBFormat InFormat, const int32 InBitDe
 	check(CompressedData.Num());
 
 	// Ensure we haven't already uncompressed the file.
-	if (RawData.Num() && InFormat == RawFormat && InBitDepth == RawBitDepth)
+	if (RawData.Num() && InFormat == Format && InBitDepth == BitDepth)
 	{
 		return;
 	}
@@ -485,8 +485,8 @@ void FExrImageWrapper::Uncompress(const ERGBFormat InFormat, const int32 InBitDe
 		}
 	}
 
-	RawFormat = InFormat;
-	RawBitDepth = InBitDepth;
+	Format = InFormat;
+	BitDepth = InBitDepth;
 }
 
 
diff --git a/Engine/Source/Runtime/ImageWrapper/Private/Formats/JpegImageWrapper.cpp b/Engine/Source/Runtime/ImageWrapper/Private/Formats/JpegImageWrapper.cpp
--- a/Engine/Source/Runtime/ImageWrapper/Private/Formats/JpegImageWrapper.cpp
+++ b/Engine/Source/Runtime/ImageWrapper/Private/Formats/JpegImageWrapper.cpp
@@ -50,7 +50,7 @@ namespace
 			case ERGBFormat::BGRA:	return TJPF_BGRA;
 			case ERGBFormat::Gray:	return TJPF_GRAY;
 			case ERGBFormat::RGBA:	return TJPF_RGBA;
-			default:				return TJPF_RGBA;
+			default:	check(0);	return TJPF_RGBA;
 		}
 	}
 }
@@ -204,7 +204,7 @@ void FJpegImageWrapper::Compress(int32 Quality)
 		check(Height > 0);
 
 		// re-order components if required - JPEGs expect RGBA
-		if(RawFormat == ERGBFormat::BGRA)
+		if(Format == ERGBFormat::BGRA)
 		{
 			FColor* Colors = (FColor*)RawData.GetData();
 			const int32 NumColors = RawData.Num() / 4;
@@ -326,8 +326,9 @@ void FJpegImageWrapper::CompressTurbo(int32 Quality)
 		check(RawData.Num());
 		check(Width > 0);
 		check(Height > 0);
+		check(BitDepth == 8);
 
-		const int PixelFormat = ConvertTJpegPixelFormat(RawFormat);
+		const int PixelFormat = ConvertTJpegPixelFormat(Format);
 		const int Subsampling = TJSAMP_420;
 		const int Flags = TJFLAG_NOREALLOC | TJFLAG_FASTDCT;
 
@@ -335,7 +336,9 @@ void FJpegImageWrapper::CompressTurbo(int32 Quality)
 		CompressedData.SetNum(OutBufferSize);
 		unsigned char* OutBuffer = CompressedData.GetData();
 
-		const bool bSuccess = tjCompress2(Compressor, RawData.GetData(), Width, RawBytesPerRow, Height, PixelFormat, &OutBuffer, &OutBufferSize, Subsampling, Quality, Flags) == 0;
+		int BytesPerRow = GetBytesPerRow();
+
+		const bool bSuccess = tjCompress2(Compressor, RawData.GetData(), Width, BytesPerRow, Height, PixelFormat, &OutBuffer, &OutBufferSize, Subsampling, Quality, Flags) == 0;
 		check(bSuccess);
 
 		CompressedData.SetNum(OutBufferSize);
diff --git a/Engine/Source/Runtime/ImageWrapper/Private/Formats/PngImageWrapper.cpp b/Engine/Source/Runtime/ImageWrapper/Private/Formats/PngImageWrapper.cpp
--- a/Engine/Source/Runtime/ImageWrapper/Private/Formats/PngImageWrapper.cpp
+++ b/Engine/Source/Runtime/ImageWrapper/Private/Formats/PngImageWrapper.cpp
@@ -218,11 +218,11 @@ void FPngImageWrapper::Compress(int32 Quality)
 			}
 
 			png_set_compression_level(png_ptr, ZlibLevel);
-			png_set_IHDR(png_ptr, info_ptr, Width, Height, RawBitDepth, (RawFormat == ERGBFormat::Gray) ? PNG_COLOR_TYPE_GRAY : PNG_COLOR_TYPE_RGBA, PNG_INTERLACE_NONE, PNG_COMPRESSION_TYPE_DEFAULT, PNG_FILTER_TYPE_DEFAULT);
+			png_set_IHDR(png_ptr, info_ptr, Width, Height, BitDepth, (Format == ERGBFormat::Gray) ? PNG_COLOR_TYPE_GRAY : PNG_COLOR_TYPE_RGBA, PNG_INTERLACE_NONE, PNG_COMPRESSION_TYPE_DEFAULT, PNG_FILTER_TYPE_DEFAULT);
 			png_set_write_fn(png_ptr, this, FPngImageWrapper::user_write_compressed, FPngImageWrapper::user_flush_data);
 
-			const uint64 PixelChannels = (RawFormat == ERGBFormat::Gray) ? 1 : 4;
-			const uint64 BytesPerPixel = (RawBitDepth * PixelChannels) / 8;
+			const uint64 PixelChannels = (Format == ERGBFormat::Gray) ? 1 : 4;
+			const uint64 BytesPerPixel = (BitDepth * PixelChannels) / 8;
 			const uint64 BytesPerRow = BytesPerPixel * Width;
 
 			for (int64 i = 0; i < Height; i++)
@@ -231,12 +231,12 @@ void FPngImageWrapper::Compress(int32 Quality)
 			}
 			png_set_rows(png_ptr, info_ptr, row_pointers);
 
-			uint32 Transform = (RawFormat == ERGBFormat::BGRA) ? PNG_TRANSFORM_BGR : PNG_TRANSFORM_IDENTITY;
+			uint32 Transform = (Format == ERGBFormat::BGRA) ? PNG_TRANSFORM_BGR : PNG_TRANSFORM_IDENTITY;
 
 			// PNG files store 16-bit pixels in network byte order (big-endian, ie. most significant bits first).
 #if PLATFORM_LITTLE_ENDIAN
 			// We're little endian so we need to swap
-			if (RawBitDepth == 16)
+			if (BitDepth == 16)
 			{
 				Transform |= PNG_TRANSFORM_SWAP_ENDIAN;
 			}
@@ -268,7 +268,7 @@ bool FPngImageWrapper::SetCompressed(const void* InCompressedData, int64 InCompr
 
 void FPngImageWrapper::Uncompress(const ERGBFormat InFormat, const int32 InBitDepth)
 {
-	if(!RawData.Num() || InFormat != RawFormat || InBitDepth != RawBitDepth)
+	if(!RawData.Num() || InFormat != Format || InBitDepth != BitDepth)
 	{
 		check(CompressedData.Num());
 		UncompressPNGData(InFormat, InBitDepth);
@@ -444,8 +444,8 @@ void FPngImageWrapper::UncompressPNGData(const ERGBFormat InFormat, const int32
 	}
 #endif
 
-	RawFormat = InFormat;
-	RawBitDepth = InBitDepth;
+	Format = InFormat;
+	BitDepth = InBitDepth;
 }
 
 
diff --git a/Engine/Source/Runtime/ImageWrapper/Private/ImageWrapperBase.cpp b/Engine/Source/Runtime/ImageWrapper/Private/ImageWrapperBase.cpp
--- a/Engine/Source/Runtime/ImageWrapper/Private/ImageWrapperBase.cpp
+++ b/Engine/Source/Runtime/ImageWrapper/Private/ImageWrapperBase.cpp
@@ -9,10 +9,7 @@
 
 FImageWrapperBase::FImageWrapperBase()
 	: Format(ERGBFormat::Invalid)
-	, RawFormat(ERGBFormat::Invalid)
 	, BitDepth(0)
-	, RawBitDepth(0)
-	, RawBytesPerRow(0)
 	, Width(0)
 	, Height(0)
 { }
@@ -25,14 +22,9 @@ void FImageWrapperBase::Reset()
 {
 	LastError.Empty();
 	
-	//@@!!
-	// not reset ?
-	//TArray64<uint8> RawData;
-	//TArray64<uint8> CompressedData;
-
-	RawFormat = ERGBFormat::Invalid;
-	RawBitDepth = 0;
-	RawBytesPerRow = 0;
+	RawData.Empty();
+	CompressedData.Empty();
+
 	Format = ERGBFormat::Invalid;
 	BitDepth = 0;
 	Width = 0;
@@ -117,6 +109,18 @@ bool FImageWrapperBase::SetRaw(const void* InRawData, int64 InRawSize, const int
 		return false;
 	}
 	
+	Format = InFormat;
+	BitDepth = InBitDepth;
+	Width = InWidth;
+	Height = InHeight;
+
+	int BytesPerRow = GetBytesPerRow();
+	
+	// stride not supported :
+	check( InBytesPerRow == 0 || BytesPerRow == InBytesPerRow );
+
+	check( InRawSize == BytesPerRow * Height );
+
 	// this is usually an unnecessary allocation and copy
 	// we should compress directly from the source buffer
 
@@ -124,19 +128,73 @@ bool FImageWrapperBase::SetRaw(const void* InRawData, int64 InRawSize, const int
 	RawData.AddUninitialized(InRawSize);
 	FMemory::Memcpy(RawData.GetData(), InRawData, InRawSize);
 
-	RawFormat = InFormat;
-	RawBitDepth = InBitDepth;
-	RawBytesPerRow = InBytesPerRow;
+	return true;
+}
 
-	Format = RawFormat;
-	BitDepth = RawBitDepth;
 
-	Width = InWidth;
-	Height = InHeight;
+int IImageWrapper::GetRGBFormatBytesPerPel(ERGBFormat RGBFormat,int BitDepth)
+{
+	switch(RGBFormat)
+	{
+	case ERGBFormat::RGBA:
+	case ERGBFormat::BGRA:
+		if ( BitDepth == 8 )
+		{
+			return 4;
+		}
+		else if ( BitDepth == 16 )
+		{
+			return 8;
+		}
+		break;
+						
+	case ERGBFormat::Gray:
+		if ( BitDepth == 8 )
+		{
+			return 1;
+		}
+		else if ( BitDepth == 16 )
+		{
+			return 2;
+		}
+		break;
+			
+	case ERGBFormat::BGRE:
+		if ( BitDepth == 8 )
+		{
+			return 4;
+		}
+		break;
+			
+	case ERGBFormat::RGBAF:
+		if ( BitDepth == 16 )
+		{
+			return 8;
+		}
+		else if ( BitDepth == 32 )
+		{
+			return 16;
+		}
+		break;
+			
+	case ERGBFormat::GrayF:
+		if ( BitDepth == 16 )
+		{
+			return 2;
+		}
+		else if ( BitDepth == 32 )
+		{
+			return 4;
+		}
+		break;
 
-	return true;
-}
+	default:
+		break;
+	}
 
+	UE_LOG(LogImageWrapper, Error, TEXT("GetRGBFormatBytesPerPel not handled : %d,%d"), (int)RGBFormat,BitDepth );
+	return 0;
+}
 
 
 ERawImageFormat::Type IImageWrapper::ConvertRGBFormat(ERGBFormat RGBFormat,int BitDepth,bool * bIsExactMatch)
diff --git a/Engine/Source/Runtime/ImageWrapper/Private/ImageWrapperBase.h b/Engine/Source/Runtime/ImageWrapper/Private/ImageWrapperBase.h
--- a/Engine/Source/Runtime/ImageWrapper/Private/ImageWrapperBase.h
+++ b/Engine/Source/Runtime/ImageWrapper/Private/ImageWrapperBase.h
@@ -107,19 +107,16 @@ public:
 
 protected:
 
+	int GetBytesPerPel() const { return GetRGBFormatBytesPerPel(Format,BitDepth); }
+	int GetBytesPerRow() const { return Width * GetBytesPerPel(); }
+
 	/** Arrays of compressed/raw data */
 	TArray64<uint8> RawData;
 	TArray64<uint8> CompressedData;
 
 	/** Format of the raw data */
-	// @@!! two format vars ugh
 	ERGBFormat Format;
-	ERGBFormat RawFormat;
-	int8 BitDepth;
-	int8 RawBitDepth;
-
-	/** Bytes per row for the raw data */
-	int32 RawBytesPerRow; // @@!! not supported, remove me
+	int BitDepth;
 
 	/** Width/Height of the image data */
 	int32 Width;
diff --git a/Engine/Source/Runtime/ImageWrapper/Public/IImageWrapper.h b/Engine/Source/Runtime/ImageWrapper/Public/IImageWrapper.h
--- a/Engine/Source/Runtime/ImageWrapper/Public/IImageWrapper.h
+++ b/Engine/Source/Runtime/ImageWrapper/Public/IImageWrapper.h
@@ -302,10 +302,12 @@ public:
 		return GetBitDepth() == 8 && GetFormat() != ERGBFormat::BGRE;
 	}
 
-	// call these from ImageWrapperModule.h , see documentation there
+	// external users call these from ImageWrapperModule.h , see documentation there
 	IMAGEWRAPPER_API static void ConvertRawImageFormat(ERawImageFormat::Type RawFormat, ERGBFormat & OutFormat,int & OutBitDepth);
 	IMAGEWRAPPER_API static ERawImageFormat::Type ConvertRGBFormat(ERGBFormat RGBFormat,int BitDepth,bool * bIsExactMatch = nullptr);
 	
+	IMAGEWRAPPER_API static int GetRGBFormatBytesPerPel(ERGBFormat RGBFormat,int BitDepth);
+
 	/* get the current image format, mapped to an ERawImageFormat
 	 * if ! *bIsExactMatch , conversion is needed
 	 * can call after SetCompressed()
