fix ImageWrapper Mac compile?

diff --git a/Engine/Source/Runtime/ImageWrapper/Private/Formats/IcnsImageWrapper.cpp b/Engine/Source/Runtime/ImageWrapper/Private/Formats/IcnsImageWrapper.cpp
--- a/Engine/Source/Runtime/ImageWrapper/Private/Formats/IcnsImageWrapper.cpp
+++ b/Engine/Source/Runtime/ImageWrapper/Private/Formats/IcnsImageWrapper.cpp
@@ -75,8 +75,8 @@ void FIcnsImageWrapper::Uncompress(const ERGBFormat InFormat, const int32 InBitD
 			RawData.Empty();
 			RawData.Append([Bitmap bitmapData], [Bitmap bytesPerPlane]);
 
-			RawFormat = Format = InFormat;
-			RawBitDepth = BitDepth = InBitDepth;
+			Format = InFormat;
+			BitDepth = InBitDepth;
 
 			Width = [Bitmap pixelsWide];
 			Height = [Bitmap pixelsHigh];
