remove imagewrapper srgb todo notes and comment to clarify existing behavior

diff --git a/Engine/Source/Runtime/ImageWrapper/Private/Formats/DdsImageWrapper.cpp b/Engine/Source/Runtime/ImageWrapper/Private/Formats/DdsImageWrapper.cpp
--- a/Engine/Source/Runtime/ImageWrapper/Private/Formats/DdsImageWrapper.cpp
+++ b/Engine/Source/Runtime/ImageWrapper/Private/Formats/DdsImageWrapper.cpp
@@ -25,7 +25,7 @@ void FDdsImageWrapper::Compress(int32 Quality)
 		return;
 	}
 
-	// @@!! SRGB
+	// we are not passed SRGB/Gamma info, just assume it is Default for now :
 	EGammaSpace GammaSpace = ERawImageFormat::GetDefaultGammaSpace(RawImageFormat);
 	
 	// some code dupe with IImageWrapper::GetRawImage
@@ -144,6 +144,10 @@ bool FDdsImageWrapper::SetCompressed(const void* InCompressedData, int64 InCompr
 		UE_LOG(LogImageWrapper, Warning, TEXT("DDS is a complex 3d/cube/array image, DdsImageWrapper will just import the first 2d surface"));
 	}
 	
+	// SRGB/Gamma : in some cases we can get SRGB info from the DDS file, so we could report that out
+	//		see EditorFactories
+	//		not doing for now
+
 	ConvertRawImageFormat(RawImageFormat, Format,BitDepth);
 
 	Width = DDS->Width;
@@ -163,7 +167,7 @@ void FDdsImageWrapper::Uncompress(const ERGBFormat InFormat, int32 InBitDepth)
 		
 	ERawImageFormat::Type RawImageFormat = UE::DDS::DXGIFormatGetClosestRawFormat(DDS->DXGIFormat);
 
-	// @@!! SRGB
+	// not setting SRGB/Gamma, just assume it is Default for now :
 	FImage RawImage(Width,Height,RawImageFormat);
 
 	if ( ! DDS->GetMipImage(RawImage,0) )
diff --git a/Engine/Source/Runtime/ImageWrapper/Private/ImageWrapperBase.cpp b/Engine/Source/Runtime/ImageWrapper/Private/ImageWrapperBase.cpp
--- a/Engine/Source/Runtime/ImageWrapper/Private/ImageWrapperBase.cpp
+++ b/Engine/Source/Runtime/ImageWrapper/Private/ImageWrapperBase.cpp
@@ -354,11 +354,9 @@ bool IImageWrapper::GetRawImage(FImage & OutImage)
 		return false;
 	}
 
-	EGammaSpace GammaSpace = EGammaSpace::Linear;
-	if ( GetSRGB() )
-	{
-		GammaSpace = EGammaSpace::sRGB;
-	}
+	// ImageWrapper RGBFormat doesn't track if pixels are Gamma/sRGB or not
+	//	just assume they are Default for now :
+	EGammaSpace GammaSpace = ERawImageFormat::GetDefaultGammaSpace(RawFormat);
 
 	if ( bExactMatch )
 	{
diff --git a/Engine/Source/Runtime/ImageWrapper/Private/ImageWrapperModule.cpp b/Engine/Source/Runtime/ImageWrapper/Private/ImageWrapperModule.cpp
--- a/Engine/Source/Runtime/ImageWrapper/Private/ImageWrapperModule.cpp
+++ b/Engine/Source/Runtime/ImageWrapper/Private/ImageWrapperModule.cpp
@@ -357,15 +357,12 @@ public:
 		{
 			// output image format may not support this pixel format
 			//	in that case conversion is needed
-			// 
-			// @@!! warn or set a flag when this conversion is lossy?
 
 			ERawImageFormat::Type NewFormat = ImageWrapper->GetSupportedRawFormat(InImage.Format);
 			check( NewFormat != InImage.Format );
 			
 			// we will write to the image wrapper using "NewFormat"
 			// do a blit to convert InImage to NewFormat
-			// 
 
 			TempImage.Init(InImage.SizeX,InImage.SizeY,NewFormat,ERawImageFormat::GetDefaultGammaSpace(NewFormat));
 			WriteImage = TempImage;
@@ -385,7 +382,9 @@ public:
 			// go ahead and write the U8 bytes without changing them
 			// I think that's probably what is intended
 			// when someone writes Linear U8 to PNG they don't want me to do a gamma transform
-			// @@!! TODO: write the gamma setting to the output when possible, eg. for DDS output
+			
+			// in some cases we could write the gamma setting to the output when possible, eg. for DDS output
+			//  not doing that for now
 		}
 
 		if ( ! ImageWrapper->SetRaw(WriteImage.RawData,WriteImage.GetImageSizeBytes(),WriteImage.SizeX,WriteImage.SizeY,RGBFormat,BitDepth) )
diff --git a/Engine/Source/Runtime/ImageWrapper/Public/IImageWrapper.h b/Engine/Source/Runtime/ImageWrapper/Public/IImageWrapper.h
--- a/Engine/Source/Runtime/ImageWrapper/Public/IImageWrapper.h
+++ b/Engine/Source/Runtime/ImageWrapper/Public/IImageWrapper.h
@@ -291,14 +291,18 @@ public:
 	 */
 	virtual ERGBFormat GetFormat() const = 0;
 	
+	
 	/* Should the pixels be treated as sRGB encoded? (or Linear)
 	* 
-	 * @@!! TODO: implement GetSRGB() on the formats that can report it (eg. DDS?)
-	 */
-	virtual bool GetSRGB() const
+	* note: ImageWrapper Format does not track if pixels are Gamma/SRGB or not
+	* assume they are ERawImageFormat::GetDefaultGammaSpace gammaspace
+	* eg. U8 is SRGB and everything else is Linear
+	*/
+	bool GetSRGB() const
 	{
 		// sRGB is guessed from bit depth
-		// 8 = on
+		// 8 = on (except BGRE)
+		// must match ERawImageFormat::GetDefaultGammaSpace
 		return GetBitDepth() == 8 && GetFormat() != ERGBFormat::BGRE;
 	}
 
diff --git a/Engine/Source/Runtime/ImageWrapper/Public/IImageWrapperModule.h b/Engine/Source/Runtime/ImageWrapper/Public/IImageWrapperModule.h
--- a/Engine/Source/Runtime/ImageWrapper/Public/IImageWrapperModule.h
+++ b/Engine/Source/Runtime/ImageWrapper/Public/IImageWrapperModule.h
@@ -11,6 +11,18 @@
 
 /**
  * Interface for image wrapper modules.
+ * 
+ * 
+ * note on SRGB/Gamma handling :
+ * it is assumed that non-U8 data is always Linear
+ * U8 data is written without changing the bytes
+ * the gamma correction of U8 bytes is NOT persisted in the file formats (for read or write)
+ * loading will always give you SRGB for U8 (ERawImageFormat::GetDefaultGammaSpace)
+ * if U8-float conversions are required, they do respect gamma space
+ * 
+ * eg. if you write U8 Linear data to EXR, it will be converted to Linear float from U8 Linear correctly
+ * if you write U8 Linear to BMP , it will write the U8 bytes unchanged, and on load it will come back in as U8 SRGB
+ * 
  */
 class IImageWrapperModule
 	: public IModuleInterface
