PngImageWrapper detect unsupported formats and return failure rather than crash

diff --git a/Engine/Source/Runtime/ImageWrapper/Private/Formats/PngImageWrapper.cpp b/Engine/Source/Runtime/ImageWrapper/Private/Formats/PngImageWrapper.cpp
--- a/Engine/Source/Runtime/ImageWrapper/Private/Formats/PngImageWrapper.cpp
+++ b/Engine/Source/Runtime/ImageWrapper/Private/Formats/PngImageWrapper.cpp
@@ -261,9 +261,27 @@ void FPngImageWrapper::Reset()
 
 bool FPngImageWrapper::SetCompressed(const void* InCompressedData, int64 InCompressedSize)
 {
-	bool bResult = FImageWrapperBase::SetCompressed(InCompressedData, InCompressedSize);
+	if ( ! FImageWrapperBase::SetCompressed(InCompressedData, InCompressedSize) )
+	{
+		return false;
+	}
 
-	return bResult && LoadPNGHeader();	// Fetch the variables from the header info
+	if ( ! LoadPNGHeader() )
+	{
+		return false;
+	}
+	
+	if ( (Format == ERGBFormat::BGRA || Format == ERGBFormat::RGBA || Format == ERGBFormat::Gray) &&
+		(BitDepth == 8 || BitDepth == 16) )
+	{
+		return true;
+	}
+	else
+	{
+		// Other formats unsupported at present
+		UE_LOG(LogImageWrapper, Warning, TEXT("PNG Unsupported Format"));
+		return false;
+	}
 }
 
 
@@ -290,8 +308,8 @@ void FPngImageWrapper::UncompressPNGData(const ERGBFormat InFormat, const int32
 	check(Height > 0);
 
 	// Note that PNGs on PC tend to be BGR
-	check(InFormat == ERGBFormat::BGRA || InFormat == ERGBFormat::RGBA || InFormat == ERGBFormat::Gray)	// Other formats unsupported at present
-	check(InBitDepth == 8 || InBitDepth == 16)	// Other formats unsupported at present
+	check(InFormat == ERGBFormat::BGRA || InFormat == ERGBFormat::RGBA || InFormat == ERGBFormat::Gray);	// Other formats unsupported at present
+	check(InBitDepth == 8 || InBitDepth == 16);	// Other formats unsupported at present
 
 	// Reset to the beginning of file so we can use png_read_png(), which expects to start at the beginning.
 	ReadOffset = 0;
