TiffImageWrapper fix importing floating point gray images as integer
add support for F32 float import without conversion
fix incorrect quantizer in color conversion

diff --git a/Engine/Source/Runtime/ImageWrapper/Private/Formats/TiffImageWrapper.cpp b/Engine/Source/Runtime/ImageWrapper/Private/Formats/TiffImageWrapper.cpp
--- a/Engine/Source/Runtime/ImageWrapper/Private/Formats/TiffImageWrapper.cpp
+++ b/Engine/Source/Runtime/ImageWrapper/Private/Formats/TiffImageWrapper.cpp
@@ -23,7 +23,6 @@ THIRD_PARTY_INCLUDES_START
 #include "tiffio.h"
 THIRD_PARTY_INCLUDES_END
 
-
 namespace UE::ImageWrapper::Private
 {
 	namespace TiffImageWrapper
@@ -31,18 +30,36 @@ namespace UE::ImageWrapper::Private
 		template<class DataTypeDest, class DataTypeSrc>
 		DataTypeDest ConvertToWriteFormat(DataTypeSrc ReadValue)
 		{
-			if constexpr (TIsSame<DataTypeDest, DataTypeSrc>::Value || TIsSame<DataTypeDest, FFloat16>::Value)
+			if constexpr ( TIsSame<DataTypeDest, DataTypeSrc>::Value )
 			{
 				return ReadValue;
 			}
-			else if constexpr (TIsSame<DataTypeSrc, FFloat16>::Value || TIsFloatingPoint<DataTypeSrc>::Value)
-			{
-				return TNumericLimits<DataTypeDest>::Max() * FMath::Clamp(ReadValue, 0.0, 1.0);
-			}
 			else
 			{
-				// Naive linear interpolation
-				return TNumericLimits<DataTypeDest>::Max() * (double(ReadValue) / double(TNumericLimits<DataTypeSrc>::Max()));
+				// we almost always load the tiff into the same format it was stored in
+				// so this conversion is rarely used
+				// one case where it is used is for int32 to int16
+
+				// get source value in [0,1]
+				double SrcValue;
+				if constexpr (TIsSame<DataTypeSrc, FFloat16>::Value || TIsFloatingPoint<DataTypeSrc>::Value)
+				{
+					SrcValue = FMath::Clamp((double)ReadValue, 0.0, 1.0);
+				}
+				else
+				{
+					SrcValue = ReadValue * (1.0 / (double)TNumericLimits<DataTypeSrc>::Max() );
+				}
+				
+				// quantize to dest :
+				if constexpr (TIsSame<DataTypeDest, FFloat16>::Value || TIsFloatingPoint<DataTypeDest>::Value)
+				{
+					return DataTypeDest(SrcValue);
+				}
+				else
+				{
+					return (DataTypeDest)( SrcValue * (double)TNumericLimits<DataTypeDest>::Max() + 0.50 );
+				}
 			}
 		}
 
@@ -413,14 +430,12 @@ namespace UE::ImageWrapper::Private
 	// CanSetRawFormat returns true if SetRaw will accept this format
 	bool FTiffImageWrapper::CanSetRawFormat(const ERGBFormat InFormat, const int32 InBitDepth) const
 	{
-		//checkf(false, TEXT("TIFF compression not supported"));
 		return false;
 	}
 
 	// returns InFormat if supported, else maps to something supported
 	ERawImageFormat::Type FTiffImageWrapper::GetSupportedRawFormat(const ERawImageFormat::Type InFormat) const
 	{
-		//checkf(false, TEXT("TIFF compression not supported"));
 		return ERawImageFormat::BGRA8;
 	}
 
@@ -428,7 +443,6 @@ namespace UE::ImageWrapper::Private
 	{
 		if (InFormat == Format && InBitDepth == BitDepth)
 		{
-
 			// Read using RGBA
 			if (Format == ERGBFormat::BGRA && BitDepth == 8)
 			{
@@ -440,6 +454,7 @@ namespace UE::ImageWrapper::Private
 
 				if (TIFFReadRGBAImageOriented(Tiff, Width, Height, static_cast<uint32*>(static_cast<void*>(RawData.GetData())), 0, ORIENTATION_LEFTTOP) != 0)
 				{
+					// @@ use ImageCore TransposeImageRGBABGRA
 					EParallelForFlags ParallelForFlags = IsInGameThread() ? EParallelForFlags::None : EParallelForFlags::BackgroundPriority;
 					ParallelFor(Height, [this](int32 HeightIndex)
 						{
@@ -468,16 +483,29 @@ namespace UE::ImageWrapper::Private
 			{
 				UnpackIntoRawBuffer<FFloat16>(4);
 			}
-			else if (Format == ERGBFormat::Gray)
+			else if (Format == ERGBFormat::RGBAF && BitDepth == 32)
 			{
-				if (BitDepth == 8)
-				{
-					UnpackIntoRawBuffer<uint8>(1);
-				}
-				else if (BitDepth == 16)
-				{
-					UnpackIntoRawBuffer<uint16>(1);
-				}
+				UnpackIntoRawBuffer<float>(4);
+			}
+			else if (Format == ERGBFormat::Gray && BitDepth == 8)
+			{
+				UnpackIntoRawBuffer<uint8>(1);
+			}
+			else if (Format == ERGBFormat::Gray && BitDepth == 16)
+			{
+				UnpackIntoRawBuffer<uint16>(1);
+			}
+			else if (Format == ERGBFormat::GrayF && BitDepth == 16)
+			{
+				UnpackIntoRawBuffer<FFloat16>(1);
+			}
+			else if (Format == ERGBFormat::GrayF && BitDepth == 32)
+			{
+				UnpackIntoRawBuffer<float>(1);
+			}
+			else
+			{
+				SetError(TEXT("Unsupported requested format for the input image. Can't uncompress the tiff image."));
 			}
 		}
 		else
@@ -528,17 +556,23 @@ namespace UE::ImageWrapper::Private
 				case PHOTOMETRIC_RGB:
 					if ( SamplesPerPixel == 3 || SamplesPerPixel == 4)
 					{
-						if (SampleFormat == SAMPLEFORMAT_IEEEFP && (BitsPerSample == 32 || BitsPerSample == 16))
+						if (BitsPerSample == 16 || BitsPerSample == 32 || BitsPerSample == 64)
 						{
-							Format = ERGBFormat::RGBAF;
-							BitDepth = 16;
-						}
-						else if (BitsPerSample == 16 || BitsPerSample == 32 || BitsPerSample == 64)
-						{
-							Format = ERGBFormat::RGBA;
-							BitDepth = 16;
+							if (SampleFormat == SAMPLEFORMAT_IEEEFP )
+							{
+								Format = ERGBFormat::RGBAF;
+								
+								BitDepth = ( BitsPerSample >= 32 ) ? 32 : 16;
+							}
+							else
+							{
+								Format = ERGBFormat::RGBA;
+
+								// we don't have a 32-bit integer texture format, so always load as 16 bit								
+								BitDepth = 16;
+							}
 						}
-						else
+						else if ( BitsPerSample <= 8 )
 						{
 							// Will be converted to 8 bits per channel
 							Format = ERGBFormat::BGRA;
@@ -557,7 +591,7 @@ namespace UE::ImageWrapper::Private
 
 				case PHOTOMETRIC_MINISBLACK:
 				case PHOTOMETRIC_MINISWHITE:
-					if (SamplesPerPixel == 1 || SamplesPerPixel == 2)
+					if (SamplesPerPixel == 1 || SamplesPerPixel == 2) // 2 ?
 					{
 						if (BitsPerSample == 1 || BitsPerSample == 2 || BitsPerSample == 4 || BitsPerSample == 8)
 						{
@@ -575,17 +609,36 @@ namespace UE::ImageWrapper::Private
 						{
 							if (SamplesPerPixel == 1)
 							{
-								Format = ERGBFormat::Gray;
+								if (SampleFormat == SAMPLEFORMAT_IEEEFP)
+								{
+									Format = ERGBFormat::GrayF;
+								}
+								else
+								{
+									Format = ERGBFormat::Gray;
+								}
 							}
-							else if (SampleFormat == SAMPLEFORMAT_IEEEFP)
+							else
 							{
-								Format = ERGBFormat::RGBAF;
+								if (SampleFormat == SAMPLEFORMAT_IEEEFP)
+								{
+									Format = ERGBFormat::RGBAF;
+								}
+								else
+								{
+									Format = ERGBFormat::RGBA;
+								}
+							}
+
+							if (SampleFormat == SAMPLEFORMAT_IEEEFP && BitsPerSample >= 32 )
+							{
+								BitDepth = 32;
 							}
 							else
 							{
-								Format = ERGBFormat::RGBA;
+								// we don't have an int32 texture format, so they convert to 16 bit
+								BitDepth = 16;
 							}
-							BitDepth = 16;
 						}
 					}
 					break;
@@ -596,7 +649,12 @@ namespace UE::ImageWrapper::Private
 						if (BitsPerSample == 1
 							|| BitsPerSample == 2
 							|| BitsPerSample == 4
-							|| BitsPerSample == 8
+							|| BitsPerSample == 8)
+						{
+							Format = ERGBFormat::BGRA;
+							BitDepth = 8;
+						}
+						else if ( BitsPerSample == 16
 							|| BitsPerSample == 32
 							|| BitsPerSample == 64)
 						{
@@ -732,7 +790,7 @@ namespace UE::ImageWrapper::Private
 				SetError(TEXT("Tiff, Unsupported bits per sample from a uint sample format."));
 			}
 		}
-		else if constexpr (TIsSame<DataTypeDest, uint16>::Value || TIsSame<DataTypeDest, FFloat16>::Value)
+		else
 		{
 			if (SampleFormat == SAMPLEFORMAT_UINT || SampleFormat == 0)
 			{
@@ -766,6 +824,10 @@ namespace UE::ImageWrapper::Private
 							break;
 						}
 					}
+					else
+					{
+						SetError(TEXT("Tiff, Unsupported bits per sample from a Palette base image."));
+					}
 				}
 				else if (BitsPerSample == 16)
 				{
@@ -804,30 +866,29 @@ namespace UE::ImageWrapper::Private
 				}
 			}
 		}
-
-		if constexpr (!TIsSame<DataTypeDest, FFloat16>::Value)
+		
+		if (Photometric == PHOTOMETRIC_MINISWHITE)
 		{
-			if (Photometric == PHOTOMETRIC_MINISWHITE)
+			if constexpr ( TIsSame<DataTypeDest, FFloat16>::Value || TIsFloatingPoint<DataTypeDest>::Value )
 			{
+				// miniswhite with floating point tiff is an odd thing to do
+				// it's unclear if you should invert against 1.f or something else?
+				UE_LOG(LogImageWrapper, Warning, TEXT("Tiff MINISWHITE floating point?  This is probably wrong."));
+
 				TArrayView64<DataTypeDest> FinalImage(static_cast<DataTypeDest*>(static_cast<void*>(RawData.GetData())), RawData.Num());
 				ParallelFor(Width * Height, [&FinalImage](int32 Index)
 					{
 						DataTypeDest& FinalValue = FinalImage[Index];
-						FinalValue = TNumericLimits<DataTypeDest>::Max() - FinalValue;
+						FinalValue = 1.f - FinalValue; // uses implicit operator float() on FFloat16
 					},  IsInGameThread() ? EParallelForFlags::None : EParallelForFlags::BackgroundPriority);
 			}
-		}
-		else
-		{
-			if (Photometric == PHOTOMETRIC_MINISWHITE)
+			else
 			{
-				FFloat16 MaxValue;
-				MaxValue.Encoded = 0x7BFF;
 				TArrayView64<DataTypeDest> FinalImage(static_cast<DataTypeDest*>(static_cast<void*>(RawData.GetData())), RawData.Num());
-				ParallelFor(Width * Height, [&FinalImage, MaxValue](int32 Index)
+				ParallelFor(Width * Height, [&FinalImage](int32 Index)
 					{
 						DataTypeDest& FinalValue = FinalImage[Index];
-						FinalValue = MaxValue - FinalValue;
+						FinalValue = TNumericLimits<DataTypeDest>::Max() - FinalValue;
 					},  IsInGameThread() ? EParallelForFlags::None : EParallelForFlags::BackgroundPriority);
 			}
 		}
@@ -1035,13 +1096,13 @@ namespace UE::ImageWrapper::Private
 			ParallelFor(Width * Height, [NumOfChannelDest, &WriteArray](int32 Index)
 				{
 					const int64 WriteIndex = int64(NumOfChannelDest) * Index + 3;
-					if constexpr (!TIsSame<DataTypeDest, FFloat16>::Value)
+					if constexpr ( TIsSame<DataTypeDest, FFloat16>::Value || TIsFloatingPoint<DataTypeDest>::Value ) 
 					{
-						WriteArray[WriteIndex]= TNumericLimits<DataTypeDest>::Max();
+						WriteArray[WriteIndex] = DataTypeDest(1.f);
 					}
 					else
 					{
-						WriteArray[WriteIndex] = DataTypeDest(1.f);
+						WriteArray[WriteIndex]= TNumericLimits<DataTypeDest>::Max();
 					}
 				}, IsInGameThread() ? EParallelForFlags::None : EParallelForFlags::BackgroundPriority);
 		}
diff --git a/Engine/Source/Runtime/ImageWrapper/Private/ImageWrapperBase.cpp b/Engine/Source/Runtime/ImageWrapper/Private/ImageWrapperBase.cpp
--- a/Engine/Source/Runtime/ImageWrapper/Private/ImageWrapperBase.cpp
+++ b/Engine/Source/Runtime/ImageWrapper/Private/ImageWrapperBase.cpp
@@ -57,6 +57,7 @@ bool FImageWrapperBase::GetRaw(const ERGBFormat InFormat, int32 InBitDepth, TArr
 
 	if ( ! LastError.IsEmpty())
 	{
+		UE_LOG(LogImageWrapper, Warning, TEXT("ImageWrapper GetRaw failed: %s"), *LastError);
 		return false;
 	}
 	
