Removed ImageWrapper private pch.
Timing when compiling module with PCH(including compiling PCH): 8.47 secs
Timing using shared PCH: 6.43 secs

diff --git a/Engine/Source/Runtime/ImageWrapper/ImageWrapper.Build.cs b/Engine/Source/Runtime/ImageWrapper/ImageWrapper.Build.cs
--- a/Engine/Source/Runtime/ImageWrapper/ImageWrapper.Build.cs
+++ b/Engine/Source/Runtime/ImageWrapper/ImageWrapper.Build.cs
@@ -65,9 +65,5 @@ public class ImageWrapper : ModuleRules
 
 		// Enable exceptions to allow error handling
 		bEnableExceptions = true;
-
-		// Disable shared PCHs to handle warning C4652
-		PCHUsage = ModuleRules.PCHUsageMode.NoSharedPCHs;
-		PrivatePCHHeaderFile = "Private/ImageWrapperPrivate.h";
 	}
 }
