Fixing missing texture post gltf import.
By adding support for bitdepth 1,2,4
(GLTF had png textures with bit depth of 1 and 2)

diff --git a/Engine/Source/Runtime/ImageWrapper/Private/Formats/PngImageWrapper.cpp b/Engine/Source/Runtime/ImageWrapper/Private/Formats/PngImageWrapper.cpp
--- a/Engine/Source/Runtime/ImageWrapper/Private/Formats/PngImageWrapper.cpp
+++ b/Engine/Source/Runtime/ImageWrapper/Private/Formats/PngImageWrapper.cpp
@@ -271,6 +271,34 @@ bool FPngImageWrapper::SetCompressed(const void* InCompressedData, int64 InCompr
 		return false;
 	}
 	
+	if ((BitDepth == 1 || BitDepth == 2 || BitDepth == 4) && ((ColorType & PNG_COLOR_MASK_ALPHA) == 0))
+	{
+		// PNG specfication:
+		//  (http://www.libpng.org/pub/png/spec/1.2/PNG-Chunks.html)
+		if ((ColorType == PNG_COLOR_TYPE_PALETTE) || (ColorType == PNG_COLOR_TYPE_GRAY))
+		{
+			//From png specification:
+			//  Note that the palette uses 8 bits (1 byte) per sample regardless of the image bit depth specification. 
+			//  In particular, the palette is 8 bits deep even when it is a suggested quantization of a 16-bit truecolor image.
+
+			// ColorType == PNG_COLOR_TYPE_PALETTE supported via:
+			//	png_set_palette_to_rgb (called in UncompressPNGData)
+
+			// ColorType == PNG_COLOR_TYPE_GRAYsupported via:
+			//	png_set_expand_gray_1_2_4_to_8 (called in UncompressPNGData)
+
+			if (!RawData.Num())
+			{
+				check(CompressedData.Num());
+				UncompressPNGData(Format, 8);
+			}
+		}
+		else if (ColorType == PNG_COLOR_TYPE_RGB)
+		{
+			//according to png specification this is not a possiblity
+		}
+	}
+	
 	if ( (Format == ERGBFormat::BGRA || Format == ERGBFormat::RGBA || Format == ERGBFormat::Gray) &&
 		(BitDepth == 8 || BitDepth == 16) )
 	{
@@ -309,7 +337,7 @@ void FPngImageWrapper::UncompressPNGData(const ERGBFormat InFormat, const int32
 
 	// Note that PNGs on PC tend to be BGR
 	check(InFormat == ERGBFormat::BGRA || InFormat == ERGBFormat::RGBA || InFormat == ERGBFormat::Gray);	// Other formats unsupported at present
-	check(InBitDepth == 8 || InBitDepth == 16);	// Other formats unsupported at present
+	check(InBitDepth == 1 || InBitDepth == 2 || InBitDepth == 4 || InBitDepth == 8 || InBitDepth == 16);	// Other formats unsupported at present
 
 	// Reset to the beginning of file so we can use png_read_png(), which expects to start at the beginning.
 	ReadOffset = 0;
@@ -350,7 +378,7 @@ void FPngImageWrapper::UncompressPNGData(const ERGBFormat InFormat, const int32
 				png_set_palette_to_rgb(png_ptr);
 			}
 
-			if ((ColorType & PNG_COLOR_MASK_COLOR) == 0 && BitDepth < 8)
+			if (((ColorType & PNG_COLOR_MASK_COLOR) == 0) && BitDepth < 8)
 			{
 				png_set_expand_gray_1_2_4_to_8(png_ptr);
 			}
