Fix for FIcnsImageWrapper not handling 24bit images.  Affected the auto screen shot loading in the Project settings window.   Should fix the FSlateTexture2DRHIRef::InitDynamicRHI crash we are seeing on macOS in the editor when opening that settings window.  Add check in FSlateTexture2DRHIRef::InitDynamicRHI to catch source data size not large enough for copy operation.

diff --git a/Engine/Source/Runtime/Engine/Private/Slate/SlateTextures.cpp b/Engine/Source/Runtime/Engine/Private/Slate/SlateTextures.cpp
--- a/Engine/Source/Runtime/Engine/Private/Slate/SlateTextures.cpp
+++ b/Engine/Source/Runtime/Engine/Private/Slate/SlateTextures.cpp
@@ -73,6 +73,8 @@ void FSlateTexture2DRHIRef::InitDynamicRHI()
 			const uint32 BlocksY = GPixelFormats[PixelFormat].GetBlockCountForHeight(Height);
 			const uint32 DataStride = BlocksX * GPixelFormats[PixelFormat].BlockBytes;
 
+			checkf((uint32)TextureData->GetRawBytes().Num() >= DataStride * BlocksY, TEXT("Not enough bytes in source TextureData to complete copy operation"));
+
 			if (Stride == DataStride)
 			{
 				FMemory::Memcpy(DestTextureData, SourceTextureData, Stride * BlocksY);
diff --git a/Engine/Source/Runtime/ImageWrapper/Private/Formats/IcnsImageWrapper.cpp b/Engine/Source/Runtime/ImageWrapper/Private/Formats/IcnsImageWrapper.cpp
--- a/Engine/Source/Runtime/ImageWrapper/Private/Formats/IcnsImageWrapper.cpp
+++ b/Engine/Source/Runtime/ImageWrapper/Private/Formats/IcnsImageWrapper.cpp
@@ -71,19 +71,111 @@ void FIcnsImageWrapper::Uncompress(const ERGBFormat InFormat, const int32 InBitD
 		{
 			check(InFormat == ERGBFormat::BGRA || InFormat == ERGBFormat::RGBA);
 			check(InBitDepth == 8);
-
-			RawData.Empty();
-			RawData.Append([Bitmap bitmapData], [Bitmap bytesPerPlane]);
-
+			
 			Format = InFormat;
 			BitDepth = InBitDepth;
 
 			Width = [Bitmap pixelsWide];
 			Height = [Bitmap pixelsHigh];
 
-			if (Format == ERGBFormat::BGRA)
+			// InBitDepth of 8 above must be the per channel value not the image pixel bit depth
+			const size_t SrcBitDepth = (Bitmap.bitsPerPixel / Bitmap.samplesPerPixel);
+			check(SrcBitDepth == 8);
+			
+			// Non compressed total image size
+			const size_t SrcImageSize = Width * Height * Bitmap.samplesPerPixel;
+			
+			// 4 channel output BGRA or RGBA
+			const size_t DestBytesPerPixel = 4;
+			const size_t DestImageSize = Width * Height * DestBytesPerPixel;
+			
+			RawData.Empty();
+			RawData.AddUninitialized(DestImageSize);
+			
+			if(SrcImageSize == DestImageSize && Bitmap.bytesPerPlane == SrcImageSize && Bitmap.numberOfPlanes == 1)
+			{
+				// Exact match direct copy
+				FMemory::Memcpy(RawData.GetData(), [Bitmap bitmapData], DestImageSize);
+			}
+			else if(Bitmap.bytesPerPlane == SrcImageSize && Bitmap.numberOfPlanes == 1)
+			{
+				// Manual copy, could be 24bit or gray scale.  Be conservative about the number of availble source channels during the copy
+				uint8* SrcData = [Bitmap bitmapData];
+				
+				#define SRC_IMAGE_INDEX ((Y * Bitmap.bytesPerRow) + (X * Bitmap.samplesPerPixel))
+				#define DEST_IMAGE_INDEX ((Y * (DestBytesPerPixel * Width)) + (X * DestBytesPerPixel))
+				
+				for(size_t Y = 0;Y < Height;++Y)
+				{
+					if(Bitmap.samplesPerPixel >= 4)
+					{
+						for(size_t X = 0;X < Width;++X)
+						{
+							size_t SrcIndex = SRC_IMAGE_INDEX;
+							size_t DestIndex = DEST_IMAGE_INDEX;
+
+							RawData[DestIndex + 0] = SrcData[SrcIndex + 0];
+							RawData[DestIndex + 1] = SrcData[SrcIndex + 1];
+							RawData[DestIndex + 2] = SrcData[SrcIndex + 2];
+							RawData[DestIndex + 3] = SrcData[SrcIndex + 3];
+						}
+					}
+					else if(Bitmap.samplesPerPixel == 3)
+					{
+						for(size_t X = 0;X < Width;++X)
+						{
+							size_t SrcIndex = SRC_IMAGE_INDEX;
+							size_t DestIndex = DEST_IMAGE_INDEX;
+
+							RawData[DestIndex + 0] = SrcData[SrcIndex + 0];
+							RawData[DestIndex + 1] = SrcData[SrcIndex + 1];
+							RawData[DestIndex + 2] = SrcData[SrcIndex + 2];
+							RawData[DestIndex + 3] = 0xff;
+						}
+					}
+					else if(Bitmap.samplesPerPixel == 2)
+					{
+						// Placeholder - not sure if this is correct
+						for(size_t X = 0;X < Width;++X)
+						{
+							size_t SrcIndex = SRC_IMAGE_INDEX;
+							size_t DestIndex = DEST_IMAGE_INDEX;
+
+							RawData[DestIndex + 0] = SrcData[SrcIndex + 0];
+							RawData[DestIndex + 1] = SrcData[SrcIndex + 1];
+							RawData[DestIndex + 2] = 0x00;
+							RawData[DestIndex + 3] = 0xff;
+						}
+					}
+					else if(Bitmap.samplesPerPixel == 1)
+					{
+						// Assume gray scale - copy single input channel to all color output channels
+						for(size_t X = 0;X < Width;++X)
+						{
+							size_t SrcIndex = SRC_IMAGE_INDEX;
+							size_t DestIndex = DEST_IMAGE_INDEX;
+
+							RawData[DestIndex + 0] = SrcData[SrcIndex + 0];
+							RawData[DestIndex + 1] = SrcData[SrcIndex + 0];
+							RawData[DestIndex + 2] = SrcData[SrcIndex + 0];
+							RawData[DestIndex + 3] = 0xff;
+						}
+					}
+				}
+				
+				#undef SRC_IMAGE_INDEX
+				#undef DEST_IMAGE_INDEX
+			}
+			else
+			{
+				// Unhandled case (e.g. Planar)
+				SetError(TEXT("FIcnsImageWrapper: Cannot uncompress, unsupported format."));
+				RawData.Reset();
+			}
+
+			if ((size_t)RawData.Num() >= DestImageSize && Format == ERGBFormat::BGRA)
 			{
-				for (int64 Index = 0; Index < [Bitmap bytesPerPlane]; Index += 4)
+				for (size_t Index = 0; Index < DestImageSize; Index += 4)
 				{
 					uint8 Byte = RawData[Index];
 					RawData[Index] = RawData[Index + 2];
