When SetRaw() is called with a InBytesPerRow that isn't tightly packed, we copy the buffer row-by-row to make it tightly packed. The typical, non-strided SetRaw still does just one MemCpy call of the entire buffer.

diff --git a/Engine/Source/Runtime/ImageWrapper/Private/ImageWrapperBase.cpp b/Engine/Source/Runtime/ImageWrapper/Private/ImageWrapperBase.cpp
--- a/Engine/Source/Runtime/ImageWrapper/Private/ImageWrapperBase.cpp
+++ b/Engine/Source/Runtime/ImageWrapper/Private/ImageWrapperBase.cpp
@@ -117,17 +117,32 @@ bool FImageWrapperBase::SetRaw(const void* InRawData, int64 InRawSize, const int
 
 	int BytesPerRow = GetBytesPerRow();
 	
-	// stride not supported :
-	check( InBytesPerRow == 0 || BytesPerRow == InBytesPerRow );
 
-	check( InRawSize == BytesPerRow * Height );
+	int64 RawSize = BytesPerRow * Height;
+	RawData.Empty(RawSize);
+	RawData.AddUninitialized(RawSize);
 
-	// this is usually an unnecessary allocation and copy
-	// we should compress directly from the source buffer
+	// copy the incoming data directly
+	if ((InBytesPerRow == 0 || BytesPerRow == InBytesPerRow) && (RawSize == InRawSize))
+	{
+		// this is usually an unnecessary allocation and copy
+		// we should compress directly from the source buffer
+
+		FMemory::Memcpy(RawData.GetData(), InRawData, InRawSize);
+
+	}
+	else
+	{
+		// The supported image formats (PNG, BMP, etc) don't support strided data (although turbo jpeg does).
+		// Therefore we de-stride the data here so we can uniformly support strided input data with all the supported
+		// image formats.
+		check(RawSize < Height* InBytesPerRow);
 
-	RawData.Empty(InRawSize);
-	RawData.AddUninitialized(InRawSize);
-	FMemory::Memcpy(RawData.GetData(), InRawData, InRawSize);
+		for (int32 y = 0; y < Height; y++)
+		{
+			FMemory::Memcpy(RawData.GetData() + (y * BytesPerRow), (uint8*)InRawData + (y * InBytesPerRow), BytesPerRow);
+		}
+	}
 
 	return true;
 }
