IcnsImageWrapper was not setting image properties in SetCompressed, causing crash in GetRawImage

diff --git a/Engine/Source/Runtime/ImageWrapper/Private/Formats/IcnsImageWrapper.cpp b/Engine/Source/Runtime/ImageWrapper/Private/Formats/IcnsImageWrapper.cpp
--- a/Engine/Source/Runtime/ImageWrapper/Private/Formats/IcnsImageWrapper.cpp
+++ b/Engine/Source/Runtime/ImageWrapper/Private/Formats/IcnsImageWrapper.cpp
@@ -17,7 +17,35 @@ FIcnsImageWrapper::FIcnsImageWrapper()
 bool FIcnsImageWrapper::SetCompressed(const void* InCompressedData, int64 InCompressedSize)
 {
 #if PLATFORM_MAC
-	return FImageWrapperBase::SetCompressed(InCompressedData, InCompressedSize);
+	if ( ! FImageWrapperBase::SetCompressed(InCompressedData, InCompressedSize) )
+	{
+		return false;
+	}
+	
+	// set image properties from header
+
+	// always read as BGRA8 regardless of the image format in the file :
+	Format = ERGBFormat::BGRA;
+	BitDepth = 8;
+	Width = Height = 0;
+
+	// get width and height from image data :
+	{
+		SCOPED_AUTORELEASE_POOL;
+
+		NSData* ImageData = [NSData dataWithBytesNoCopy:CompressedData.GetData() length:CompressedData.Num() freeWhenDone:NO];
+		NSImage* Image = [[NSImage alloc] initWithData:ImageData];
+		if (Image)
+		{
+			Width = Image.size.width;
+			Height = Image.size.height;
+
+			// TODO: We have decoded the image this far we might want to store this data/object so ::Uncompress doesn't have to do it again
+			[Image release];
+		}
+	}
+
+	return ( Width > 0 );
 #else
 	return false;
 #endif
