ImageWrapper.h check format set in GetRaw to Array
could be invalid if SetCompressed failed

diff --git a/Engine/Source/Runtime/ImageWrapper/Public/IImageWrapper.h b/Engine/Source/Runtime/ImageWrapper/Public/IImageWrapper.h
--- a/Engine/Source/Runtime/ImageWrapper/Public/IImageWrapper.h
+++ b/Engine/Source/Runtime/ImageWrapper/Public/IImageWrapper.h
@@ -170,21 +170,31 @@ public:
 	* 
 	* @param OutRawData    Filled with raw image data.
 	*
-	* prefer GetRawImage instead.
+	* This GetRaw call replaces the variants with format/depth arguments, but prefer GetRawImage instead.
 	*/
 	bool GetRaw(TArray64<uint8>& OutRawData)
 	{
 		ERGBFormat Format = GetFormat();
 		int32 BitDepth = GetBitDepth();
+
+		// Format and BitDepth should have been set by SetCompressed
+		if ( Format == ERGBFormat::Invalid || BitDepth == 0 )
+		{
+			return false;
+		}
+
 		return GetRaw(Format,BitDepth,OutRawData);
 	}
 	
 	/* Decode the image file data from SetCompressed() into an FImage
 	* OutImage is allocated and attributes are filled
-	*	 passed-in contents of OutImage are destroyed
-	*	OutImage.Format is ignored
+	*	 any previous passed-in contents of OutImage are destroyed
+	*	OutImage.Format is ignored, a new format is set from the loaded image
 	*
 	* @param OutImage	 Filled with the image
+	*
+	* This is the recommended API to get the raw image data from the imagewrapper.
+	* Prefer this instead of any of the GetRaw() calls.
 	*/
 	bool GetRawImage(FImage & OutImage);
 
