FDdsImageWrapper restrict to 2D only
more complex DDS is imported through the ImportDDS path in EditorFactories
FDdsImageWrapper is only load for generic image loading and fallback path

diff --git a/Engine/Source/Runtime/ImageWrapper/Private/Formats/DdsImageWrapper.cpp b/Engine/Source/Runtime/ImageWrapper/Private/Formats/DdsImageWrapper.cpp
--- a/Engine/Source/Runtime/ImageWrapper/Private/Formats/DdsImageWrapper.cpp
+++ b/Engine/Source/Runtime/ImageWrapper/Private/Formats/DdsImageWrapper.cpp
@@ -141,7 +141,12 @@ bool FDdsImageWrapper::SetCompressed(const void* InCompressedData, int64 InCompr
 
 	if ( ! DDS->IsValidTexture2D() )
 	{
-		UE_LOG(LogImageWrapper, Warning, TEXT("DDS is a complex 3d/cube/array image, DdsImageWrapper will just import the first 2d surface"));
+		// @todo Oodle : fix me : FImageWrapperBase::SetError doesn't actually get logged anywhere
+		SetError(TEXT("DDS is a complex 3d/cube/array image, only 2d DDS surfaces are supported here"));
+
+		RawData.Empty();
+		FreeDDS();
+		return false;
 	}
 	
 	// SRGB/Gamma : in some cases we can get SRGB info from the DDS file, so we could report that out
