HdrImageWrapper major cleanups and validation fixes

1. Fix header parsing to check for actual header termination (blank line)
2. Header lines are only ever terminated by LF, never CR. Treat NUL as error.
3. The only supported pixel layout is -Y <Height> +X (the default). Require that explicitly.
4. Require that widths and heights be positive and enforce a maximum image size
5. Fix several potential integer overflows
6. Fix several potential out-of-bounds reads/writes
7. Don't cast possibly unaligned uint8 pointers to FColor*
8. Fix potential UB when shift amount in old RLE encoding gets large
9. Don't allow runs as first thing in an old-RLE scan line, since they want to repeat the prev pixel
10. Factor out a lot of the error handling code, make sure error messages get logged
11. Fix several typos while I'm here

diff --git a/Engine/Source/Runtime/ImageWrapper/Private/Formats/HdrImageWrapper.cpp b/Engine/Source/Runtime/ImageWrapper/Private/Formats/HdrImageWrapper.cpp
--- a/Engine/Source/Runtime/ImageWrapper/Private/Formats/HdrImageWrapper.cpp
+++ b/Engine/Source/Runtime/ImageWrapper/Private/Formats/HdrImageWrapper.cpp
@@ -7,14 +7,19 @@
 
 namespace UE::ImageWrapper::Private::HdrImageWrapper
 {
-	FText GetEndOFBufferErrorMessage()
+	FText GetMalformedHeaderErrorMessage()
 	{
-		return LOCTEXT("EndOFBufferError","Reached the end of the buffer before finishing decompressing the hdr. The hdr image is likely corrupted");
+		return LOCTEXT("MalformedHeader", "The file header is malformed. The HDR image is likely corrupted.");
 	}
 
-	FText GetImageDoneBeforeEndOfBufferErrorMessage()
+	FText GetEndOfBufferErrorMessage()
 	{
-		return LOCTEXT("IMageDoneButThereIsStilSomeDataToDecompress", "Reached the end of the raw image before finishing decompressing the hdr. The hdr image is likely to be corrupted");
+		return LOCTEXT("EndOFBufferError", "Reached the end of the buffer before finishing decompressing the HDR. The HDR image is likely corrupted.");
+	}
+
+	FText GetMalformedScanlineErrorMessage()
+	{
+		return LOCTEXT("MalformedScanline", "Compressed data for HDR scanline is malformed. The HDR image is likely to be corrupted.");
 	}
 }
 
@@ -26,68 +31,99 @@ bool FHdrImageWrapper::SetCompressedFromView(TArrayView64<const uint8> Data)
 
 	if (CompressedData.Num() < 11)
 	{
-		FreeCompressedData();
-		return false;
+		return FailHeaderParsing();
 	}
 
 	const uint8* FileDataPtr = CompressedData.GetData();
-
 	char Line[256];
 
-	GetHeaderLine(FileDataPtr, Line);
+	if (!GetHeaderLine(FileDataPtr, Line))
+	{
+		return FailHeaderParsing();
+	}
 
 	if (FCStringAnsi::Strcmp(Line, "#?RADIANCE") != 0 &&
 		FCStringAnsi::Strcmp(Line, "#?RGBE") != 0)
 	{
-		FreeCompressedData();
-		return false;
+		return FailHeaderParsing();
 	}
 
-	// Todo add validation for the format
-	bool bHasRGBEFormat = false;
-	while (GetHeaderLine(FileDataPtr, Line) && !bHasRGBEFormat)
+	// Read header lines: free-form, keep going until we hit a blank line
+	bool bHasFormat = false;
+
+	for (;;)
 	{
-		if (!FCStringAnsi::Strcmp(Line,"FORMAT=32-bit_rle_rgbe"))
+		if (!GetHeaderLine(FileDataPtr, Line))
 		{
-			bHasRGBEFormat = true;
+			return FailHeaderParsing();
 		}
-	}
 
-	if (bHasRGBEFormat)
-	{
-		while (GetHeaderLine(FileDataPtr, Line))
-		{
-			char* HeightStr = FCStringAnsi::Strstr(Line, "Y ");
-			char* WidthStr = FCStringAnsi::Strstr(Line, "X ");
+		// Blank line denotes end of header
+		if (!Line[0])
+			break;
 
-			if (HeightStr != NULL && WidthStr != NULL)
-			{
-				// insert a /0 after the height value
-				*(WidthStr-2) = 0;
+		const char* Cursor = Line;
 
-				Height = FCStringAnsi::Atoi(HeightStr+2);
-				Width = FCStringAnsi::Atoi(WidthStr+2);
+		// Format specified?
+		if (ParseMatchString(Cursor, "FORMAT="))
+		{
+			bHasFormat = true;
 
-				RGBDataStart = FileDataPtr;
-				return true;
+			// Currently we only support RGBE
+			if (FCStringAnsi::Strcmp(Cursor, "32-bit_rle_rgbe") != 0)
+			{
+				SetAndLogError(LOCTEXT("WrongFormatError", "The HDR image uses an unsupported format. Only the 32-bit_rle_rgbe format is supported."));
+				FreeCompressedData();
+				return false;
 			}
 		}
 	}
-	else
+
+	// If we got through the header without it mentioning a format, the file is malformed.
+	if (!bHasFormat)
 	{
-		ErrorMessage = LOCTEXT("WrongFormatError", "The hdr buffer use a unsupported format. Only the 32-bit_rle_rgbe format is supported.");
+		return FailHeaderParsing();
 	}
-	
-	FreeCompressedData();
-	return false;
+
+	// Read one more line which specifies the resolution
+	if (!GetHeaderLine(FileDataPtr, Line))
+	{
+		return FailHeaderParsing();
+	}
+
+	// If we allow enormous images, we risk int overflows on the size calcs, even with int64s.
+	// As of this writing (Oct 2022) there is no demand for reading huge HDR files, and allocating
+	// memory to enormous images is going to make the editor choke besides. Limiting
+	// pixel counts to 65536*65536 (which caps images at 16GiB) seems reasonable for now.
+	const int MaxImageDimension = 65536;
+
+	int ImageWidth;
+	int ImageHeight;
+	if (!ParseImageSize(Line, &ImageWidth, &ImageHeight) ||
+		ImageWidth <= 0 || ImageWidth > MaxImageDimension ||
+		ImageHeight <= 0 || ImageHeight > MaxImageDimension)
+	{
+		// If we don't like the resolution line (our parser is very strict), log what it was
+		// as a breadcrumb for debugging.
+		FString BadResolutionLine(Line);
+		UE_LOG(LogImageWrapper, Display, TEXT("HDR bad resolution line was: \"%s\""), *BadResolutionLine);
+
+		SetAndLogError(LOCTEXT("InvalidSizeError", "The HDR image specifies an invalid size."));
+		FreeCompressedData();
+		return false;
+	}
+
+	Width = ImageWidth;
+	Height = ImageHeight;
+	RGBDataStart = FileDataPtr;
+	return true;
 }
 
 bool FHdrImageWrapper::SetCompressed(const void* InCompressedData, int64 InCompressedSize)
 {
 	// takes copy
 	CompressedDataHolder.Reset(InCompressedSize);
-	CompressedDataHolder.AddUninitialized(InCompressedSize);
-	FMemory::Memcpy(CompressedDataHolder.GetData(), InCompressedData, InCompressedSize);
+	CompressedDataHolder.Append((const uint8*)InCompressedData, InCompressedSize);
 
 	return SetCompressedFromView(MakeArrayView(CompressedDataHolder));
 }
@@ -113,6 +149,12 @@ bool FHdrImageWrapper::SetRaw(const void* InRawData, int64 InRawSize, const int3
 		return false;
 	}
 
+	if (InWidth <= 0 || InHeight <= 0)
+	{
+		UE_LOG(LogImageWrapper, Warning, TEXT("ImageWrapper HDR unsupported size %d x %d"), InWidth, InHeight);
+		return false;
+	}
+
 	RawDataHolder.Empty(InRawSize);
 	RawDataHolder.Append((const uint8 *)InRawData,InRawSize);
 	Width = InWidth;
@@ -127,35 +169,33 @@ bool FHdrImageWrapper::SetRaw(const void* InRawData, int64 InRawSize, const int3
 TArray64<uint8> FHdrImageWrapper::GetCompressed(int32 Quality)
 {
 	// must have set BGRE8 raw data :
-	int64 NumPixels = Width * Height;
+	int64 NumPixels = (int64)Width * Height;
 	check( RawDataHolder.Num() == NumPixels * 4 );
-
 	
-	const int32 MaxHeaderSize = 256;
 	char Header[MAX_SPRINTF];
-	FCStringAnsi::Sprintf(Header, "#?RADIANCE\nFORMAT=32-bit_rle_rgbe\n\n-Y %d +X %d\n", Height,Width);
-	Header[MaxHeaderSize - 1] = 0;
-	int32 HeaderLen = FMath::Min(FCStringAnsi::Strlen(Header), MaxHeaderSize);
+	int32 HeaderLen = FCStringAnsi::Sprintf(Header, "#?RADIANCE\nFORMAT=32-bit_rle_rgbe\n\n-Y %d +X %d\n", Height, Width);
+	if (HeaderLen <= 0) // Error during Sprintf for whatever reason?
+	{
+		return TArray64<uint8>();
+	}
 
 	FreeCompressedData();
 	TArray64<uint8> CompressedDataArray;
-	CompressedDataArray.SetNum( HeaderLen + NumPixels * 4 );
+	CompressedDataArray.SetNumUninitialized(HeaderLen + NumPixels * 4);
 	uint8 * CompressedPtr = CompressedDataArray.GetData();
 
 	memcpy(CompressedPtr,Header,HeaderLen);
 	CompressedPtr += HeaderLen;
 
 	// just put the BGRE8 bytes
-	// but need to swizzle to RGBE8 order :
-
-	FColor * ToColors = (FColor *)CompressedPtr;
-	const FColor * FromColors = (const FColor *)RawDataHolder.GetData();
+	// but need to swizzle to RGBE8 order:
+	const uint8* FromBytes = RawDataHolder.GetData();
 	for(int64 i=0;i<NumPixels;i++)
 	{
-		ToColors[i].B = FromColors[i].R;
-		ToColors[i].G = FromColors[i].G;
-		ToColors[i].R = FromColors[i].B;
-		ToColors[i].A = FromColors[i].A;
+		CompressedPtr[i * 4 + 0] = FromBytes[i * 4 + 2];
+		CompressedPtr[i * 4 + 1] = FromBytes[i * 4 + 1];
+		CompressedPtr[i * 4 + 2] = FromBytes[i * 4 + 0];
+		CompressedPtr[i * 4 + 3] = FromBytes[i * 4 + 3];
 	}
 
 	return MoveTemp(CompressedDataArray);
@@ -165,7 +205,7 @@ bool FHdrImageWrapper::GetRaw(const ERGBFormat InFormat, int32 InBitDepth, TArra
 {
 	if (InFormat != ERGBFormat::BGRE || InBitDepth != 8)
 	{
-		ErrorMessage = LOCTEXT("UnSupportedFormatORBitDepth", "The format and/or the bit depth is not supported by the HdrImageWrapper. Only the BGRE format and a bitdepth of 8 is supported");
+		SetAndLogError(LOCTEXT("UnSupportedFormatORBitDepth", "The format and/or the bit depth is not supported by the HdrImageWrapper. Only the BGRE format and a bitdepth of 8 is supported"));
 		return false;
 	}
 
@@ -174,16 +214,18 @@ bool FHdrImageWrapper::GetRaw(const ERGBFormat InFormat, int32 InBitDepth, TArra
 		return false;
 	}
 
-	const int64 SizeRawImageInBytes = Width * Height * 4;
+	const int64 SizeRawImageInBytes = (int64)Width * Height * 4;
 	OutRawData.Reset(SizeRawImageInBytes);
 	OutRawData.AddUninitialized(SizeRawImageInBytes);
 
 	const uint8* FileDataPtr = RGBDataStart;
+	const uint8* FileDataEnd = CompressedData.GetData() + CompressedData.Num();
 
 	for (int32 Y = 0; Y < Height; ++Y)
 	{
-		if (!DecompressScanline(&(OutRawData[Width * Y * 4]), FileDataPtr))
+		if (!DecompressScanline(&(OutRawData[(int64)Width * Y * 4]), FileDataPtr, FileDataEnd))
 		{
+			OutRawData.Empty();
 			return false;
 		}
 	}
@@ -216,70 +258,212 @@ const FText& FHdrImageWrapper::GetErrorMessage() const
 	return ErrorMessage;
 }
 
-bool FHdrImageWrapper::GetHeaderLine(const uint8*& BufferPos, char Line[256])
+void FHdrImageWrapper::SetAndLogError(const FText& InText)
+{
+	ErrorMessage = InText;
+	UE_LOG(LogImageWrapper, Error, TEXT("%s"), *InText.ToString());
+}
+
+bool FHdrImageWrapper::FailHeaderParsing()
 {
-	char* LinePtr = Line;
+	SetAndLogError(UE::ImageWrapper::Private::HdrImageWrapper::GetMalformedHeaderErrorMessage());
+	FreeCompressedData();
+	return false;
+}
 
+bool FHdrImageWrapper::FailUnexpectedEOB()
+{
+	SetAndLogError(UE::ImageWrapper::Private::HdrImageWrapper::GetEndOfBufferErrorMessage());
+	return false;
+}
+
+bool FHdrImageWrapper::FailMalformedScanline()
+{
+	SetAndLogError(UE::ImageWrapper::Private::HdrImageWrapper::GetMalformedScanlineErrorMessage());
+	return false;
+}
+
+bool FHdrImageWrapper::GetHeaderLine(const uint8*& BufferPos, char Line[256])
+{
 	const uint8* EndOfBuffer = CompressedData.GetData() + CompressedData.Num();
-	uint32 i;
 
-	for(i = 0; i < 255; ++i)
+	for(int i = 0; i < 256; ++i)
 	{
-		if(*BufferPos == 0 || *BufferPos == 10 || *BufferPos == 13)
+		// May never read past end of buffer
+		check(BufferPos <= EndOfBuffer); // == may happen, > is a bug.
+
+		if (BufferPos >= EndOfBuffer)
 		{
-			++BufferPos;
-			break;
+			return false;
 		}
 
-		if (BufferPos >= EndOfBuffer)
+		if (*BufferPos == 0)
+		{
+			// We don't allow NUL characters in header, that's malformed
+			return false;
+		}
+		else if (*BufferPos == '\n') // Line break -> end of line in HDR (but CRs are not significant!)
+		{
+			// Terminate the line, we're good to go, but do consume the newline.
+			// We insert the terminating 0 here, and this is the only path where we return true,
+			// thus guaranteeing that successfully returned lines are 0-terminated.
+			BufferPos++;
+			Line[i] = 0;
+			return true;
+		}
+
+		// All other characters go into the line verbatim.
+		Line[i] = *BufferPos++;
+	}
+
+	// Falling out of the loop after reading 256 chars, we consider a malformed line.
+	return false;
+}
+
+// InOutCursor points to current parse cursor, into a 0-terminated string.
+// InExpected is a 0-terminated string that we expect an exact match to.
+//
+// On success, return true and point InOutCursor at the end of the matched string.
+// On failure, return false and leave InOutCursor where it is.
+bool FHdrImageWrapper::ParseMatchString(const char*& InOutCursor, const char* InExpected)
+{
+	const char* Cursor = InOutCursor;
+
+	while (*InExpected)
+	{
+		if (*Cursor != *InExpected)
 		{
-			ErrorMessage = LOCTEXT("RechedEndOfBufferWhileParsingHeader", "Reached the end of the Hdr buffer before we were done reading the header. The Hdr is invalid");
 			return false;
 		}
 
-		*LinePtr++ = *BufferPos++;
+		// The two bytes match, advance.
+		++Cursor;
+		++InExpected;
+	}
+
+	InOutCursor = Cursor;
+	return true;
+}
+
+// Like atoi, but we return the final cursor, have an error reporting mechanism
+// for overflows, and don't accept signed numbers (or '+' signs for that matter).
+//
+// InOutCursor points into a properly 0-terminated string, so we can just keep
+// reading while chars are between '0' and '9' (since NUL is not).
+bool FHdrImageWrapper::ParsePositiveInt(const char*& InOutCursor, int* OutValue)
+{
+	// We require the string to start with a digit
+	if (*InOutCursor < '0' || *InOutCursor > '9')
+	{
+		return false;
 	}
 
-	Line[i] = 0;
+	int Value = *InOutCursor - '0'; // can't overflow.
+	++InOutCursor;
+
+	// Keep consuming digits in the digit string
+	while (*InOutCursor >= '0' && *InOutCursor <= '9')
+	{
+		int Digit = *InOutCursor - '0';
+		++InOutCursor;
+
+		int64 NewValue = (int64)Value * 10 + Digit;
 
+		// Overflow?
+		if (NewValue > TNumericLimits<int>::Max())
+		{
+			return false;
+		}
+
+		// We just checked it's in range.
+		Value = (int)NewValue;
+	}
+
+	*OutValue = Value;
 	return true;
 }
 
-bool FHdrImageWrapper::DecompressScanline(uint8* Out, const uint8*& In)
+// InLine is known 0-terminated.
+bool FHdrImageWrapper::ParseImageSize(const char* InLine, int* OutWidth, int* OutHeight)
 {
-	// minimum and maxmimum scanline length for encoding
+	// We only support the (default) -Y <height> +X <width> form of the image size.
+	if (!ParseMatchString(InLine, "-Y "))
+	{
+		return false;
+	}
+
+	if (!ParsePositiveInt(InLine, OutHeight))
+	{
+		return false;
+	}
+
+	if (!ParseMatchString(InLine, " +X "))
+	{
+		return false;
+	}
+
+	if (!ParsePositiveInt(InLine, OutWidth))
+	{
+		return false;
+	}
+
+	return *InLine == 0; // All bytes in line must be consumed
+}
+
+// Trivial helper func to make repetitive error checks easier to read.
+// InCursor <= InEnd is assumed; check that we can read InAmount more bytes without
+// going past InEnd.
+bool FHdrImageWrapper::HaveBytes(const uint8* InCursor, const uint8* InEnd, int InAmount)
+{
+	// InCursor <= InEnd; InEnd - InCursor gives us the number of bytes left.
+	// (Do it this way instead of "InCursor + InAmount <= InEnd" because the latter
+	// might overflow.)
+	return InEnd - InCursor >= InAmount;
+}
+
+bool FHdrImageWrapper::DecompressScanline(uint8* Out, const uint8*& In, const uint8* InEnd)
+{
+	// minimum and maximum scanline length for encoding
 	const int32 MINELEN = 8;
 	const int32 MAXELEN = 0x7fff;
 
-	if(Width < MINELEN || Width > MAXELEN)
+	if (Width < MINELEN || Width > MAXELEN)
 	{
-		return OldDecompressScanline(Out, In, Width);
+		return OldDecompressScanline(Out, In, InEnd, Width, false);
+	}
+
+	if (!HaveBytes(In, InEnd, 1))
+	{
+		return false;
 	}
 
 	uint8 Red = *In;
 
 	if(Red != 2)
 	{
-		return OldDecompressScanline(Out, In, Width);
+		return OldDecompressScanline(Out, In, InEnd, Width, false);
 	}
 
 	++In;
 
+	if (!HaveBytes(In, InEnd, 3))
+	{
+		return false;
+	}
+
 	uint8 Green = *In++;
 	uint8 Blue = *In++;
-	uint8 Exponant = *In++;
+	uint8 Exponent = *In++;
 
-	if(Green != 2 || Blue & 128)
+	if(Green != 2 || (Blue & 128))
 	{
 		*Out++ = Blue;
 		*Out++ = Green;
 		*Out++ = Red;
-		*Out++ = Exponant;
-		return OldDecompressScanline(Out, In, Width - 1);
+		*Out++ = Exponent;
+		return OldDecompressScanline(Out, In, InEnd, Width - 1, true);
 	}
 
-	const uint8* EndOfOutBuffer = Out + Width * Height;
-	const uint8* EndOfBuffer = CompressedData.GetData() + CompressedData.Num();
 	for(uint32 ChannelRead = 0; ChannelRead < 4; ++ChannelRead)
 	{
 		// The file is in RGBE but we decompress in BGRE So swap the red and blue
@@ -293,36 +477,37 @@ bool FHdrImageWrapper::DecompressScanline(uint8* Out, const uint8*& In)
 			CurrentToWrite = 0;
 		}
 
+		const uint8* LocalIn = In;
 		uint8* OutSingleChannel = Out + CurrentToWrite;
 		int32 MultiRunIndex = 0;
+
 		while ( MultiRunIndex < Width )
 		{
-			if (In >= EndOfBuffer)
+			if (!HaveBytes(LocalIn, InEnd, 1))
 			{
-				ErrorMessage = UE::ImageWrapper::Private::HdrImageWrapper::GetEndOFBufferErrorMessage();
-				return false;
+				return FailUnexpectedEOB();
 			}
 
-			uint8 Current = *In++;
+			uint8 Current = *LocalIn++;
 
 			if (Current > 128)
 			{
-				uint32 Count = Current & 0x7f;
+				// Actual run
+				int Count = Current & 0x7f;
 
-				if (In >= EndOfBuffer)
+				if (!HaveBytes(LocalIn, InEnd, 1))
 				{
-					ErrorMessage = UE::ImageWrapper::Private::HdrImageWrapper::GetEndOFBufferErrorMessage();
-					return false;
+					return FailUnexpectedEOB();
 				}
-				Current = *In++;
+				Current = *LocalIn++;
 
-				if ( OutSingleChannel + 4 * (Count - 1) >= EndOfOutBuffer)
+				// Run needs to stay within scan line.
+				if (Width - MultiRunIndex < Count)
 				{
-					ErrorMessage = UE::ImageWrapper::Private::HdrImageWrapper::GetImageDoneBeforeEndOfBufferErrorMessage();
-					return false;
+					return FailMalformedScanline();
 				}
 
-				for(uint32 RunIndex = 0; RunIndex < Count; ++RunIndex)
+				for(int RunIndex = 0; RunIndex < Count; ++RunIndex)
 				{
 					*OutSingleChannel = Current;
 					OutSingleChannel += 4;
@@ -331,65 +516,98 @@ bool FHdrImageWrapper::DecompressScanline(uint8* Out, const uint8*& In)
 			}
 			else
 			{
-				uint32 Count = Current;
+				// Literal run.
+				int Count = Current;
 
-				if ( OutSingleChannel + 4 * (Count - 1) >= EndOfOutBuffer)
+				// Do one check up front whether we have enough data bytes following
+				if (!HaveBytes(LocalIn, InEnd, Count))
 				{
-					ErrorMessage = UE::ImageWrapper::Private::HdrImageWrapper::GetImageDoneBeforeEndOfBufferErrorMessage();
-					return false;
+					return FailUnexpectedEOB();
 				}
 
-				for(uint32 RunIndex = 0; RunIndex < Count; ++RunIndex)
+				// Literal run needs to stay within scan line.
+				if (Width - MultiRunIndex < Count)
 				{
-					if (In >= EndOfBuffer)
-					{
-						ErrorMessage = UE::ImageWrapper::Private::HdrImageWrapper::GetEndOFBufferErrorMessage();
-						return false;
-					}
-					*OutSingleChannel = *In++;
+					return FailMalformedScanline();
+				}
+
+				for(int RunIndex = 0; RunIndex < Count; ++RunIndex)
+				{
+					// All buffer checks were done up front.
+					*OutSingleChannel = *LocalIn++;
 					OutSingleChannel += 4;
 				}
 				MultiRunIndex += Count;
 			}
 		}
+
+		In = LocalIn;
 	}
 
 	return true;
 }
 
-bool FHdrImageWrapper::OldDecompressScanline(uint8* Out, const uint8*& In, int32 Lenght)
+bool FHdrImageWrapper::OldDecompressScanline(uint8* Out, const uint8*& InCodedScanline, const uint8* InEnd, int32 Length, bool bInitialRunAllowed)
 {
+	const uint8* In = InCodedScanline; // Copy to local var
 	int32 Shift = 0;
 
-	while(Lenght > 0)
+	// If an initial run is not allowed, set Shift to 32, which will make us fail if the first thing we
+	// see in this scanline is a run, but will be cleared after the first non-run pixel.
+	if (!bInitialRunAllowed)
 	{
+		Shift = 32;
+	}
+
+	while (Length > 0)
+	{
+		if (!HaveBytes(In, InEnd, 4))
+		{
+			return FailUnexpectedEOB();
+		}
+
 		uint8 Red = *In++; 
 		uint8 Green = *In++;
 		uint8 Blue = *In++; 
 		uint8 Exponent = *In++; 
 
-		if(Red == 1 && Green ==1 && Blue == 1)
+		if(Red == 1 && Green == 1 && Blue == 1)
 		{
-			int32 Count = ((int32)(Exponent) << Shift);
-			Lenght -= Count;
+			// It's not illegal to hit Shift=32 (say after a non-trivial run with Shift=24 in a giant image),
+			// but since we have a 32-bit width limit, there is just no legitimate reason to have another run
+			// once Shift=32, it should always be followed by literals.
+			//
+			// We also use this to catch runs at the start of a scanline when there's no pixels to repeat yet,
+			// by initializing Shift=32 on the first pixel of a scanline (we sometimes get called with one pixel
+			// already decoded, so this is conditional).
+			if (Shift >= 32)
+			{
+				return FailMalformedScanline();
+			}
 
-			if ( Lenght < 0 )
+			// Doing Count calculation in 64 bits to avoid overflow concerns when Shift=24.
+			int64 Count = (int64)Exponent << Shift;
+			if (Count > Length)
 			{
-				ErrorMessage = LOCTEXT("EndOfLineError","Reached the end of the outputted scanline before finishing decompressing the line. The hdr image is likely to be corrupted");
-				return false;
+				return FailMalformedScanline();
 			}
 
+			Length -= Count;
+
+			// Read previous pixel. See comments on top of function and handling of Shift >= 32 above for why
+			// we are guaranteed to have a previous pixel in the scanline when we get here.
 			Red = *(Out - 4); 
 			Green = *(Out - 3);
 			Blue = *(Out - 2); 
 			Exponent = *(Out - 1); 
 
-			for(int32 i = Count; i > 0; --i)
+			while (Count > 0)
 			{
 				*Out++ = Blue;
 				*Out++ = Green;
 				*Out++ = Red;
 				*Out++ = Exponent;
+				--Count;
 			}
 
 			Shift += 8;
@@ -401,10 +619,12 @@ bool FHdrImageWrapper::OldDecompressScanline(uint8* Out, const uint8*& In, int32
 			*Out++ = Red;
 			*Out++ = Exponent;
 			Shift = 0;
-			--Lenght;
+			--Length;
 		}
 	}
 
+	// On successful decode, copy read cursor back
+	InCodedScanline = In;
 	return true;
 }
 
diff --git a/Engine/Source/Runtime/ImageWrapper/Public/Formats/HdrImageWrapper.h b/Engine/Source/Runtime/ImageWrapper/Public/Formats/HdrImageWrapper.h
--- a/Engine/Source/Runtime/ImageWrapper/Public/Formats/HdrImageWrapper.h
+++ b/Engine/Source/Runtime/ImageWrapper/Public/Formats/HdrImageWrapper.h
@@ -44,12 +44,24 @@ public:
 
 	using IImageWrapper::GetRaw;
 private:
+	// Helpers for error exits. Set error messages and return false.
+	void SetAndLogError(const FText& InText);
+	bool FailHeaderParsing(); // also calls FreeCompressData.
+	bool FailUnexpectedEOB();
+	bool FailMalformedScanline();
+
 	bool GetHeaderLine(const uint8*& BufferPos, char Line[256]);
 
+	static bool ParseMatchString(const char*& InOutCursor, const char* InExpected);
+	static bool ParsePositiveInt(const char*& InOutCursor, int* OutValue);
+	static bool ParseImageSize(const char* InLine, int* OutWidth, int* OutHeight);
+
+	static bool HaveBytes(const uint8* InCursor, const uint8* InEnd, int InAmount);
+
 	/** @param Out order in bytes: RGBE */
-	bool DecompressScanline(uint8* Out, const uint8*& In);
+	bool DecompressScanline(uint8* Out, const uint8*& In, const uint8* InEnd);
 
-	bool OldDecompressScanline(uint8* Out, const uint8*& In, int32 Lenght);
+	bool OldDecompressScanline(uint8* Out, const uint8*& In, const uint8* InEnd, int32 Length, bool bInitialRunAllowed);
 
 	bool IsCompressedImageValid() const;
 
