BmpImageWrapper hardening
also fix several cases where bmp image data was imported wrong

diff --git a/Engine/Source/Runtime/ImageWrapper/Private/Formats/BmpImageWrapper.cpp b/Engine/Source/Runtime/ImageWrapper/Private/Formats/BmpImageWrapper.cpp
--- a/Engine/Source/Runtime/ImageWrapper/Private/Formats/BmpImageWrapper.cpp
+++ b/Engine/Source/Runtime/ImageWrapper/Private/Formats/BmpImageWrapper.cpp
@@ -5,6 +5,15 @@
 
 #include "BmpImageSupport.h"
 
+static inline bool BmpDimensionIsValid(int32 Dim)
+{
+	//ensures we can do Width*4 in int32
+	//	and then also call Align() on it
+	const int32 BmpMaxDimension = (INT32_MAX/4) - 2;
+
+	// zero dimension could be technically valid, but we won't load it
+	return Dim > 0 && Dim <= BmpMaxDimension;
+}
 
 /**
  * BMP image wrapper class.
@@ -16,6 +25,7 @@ FBmpImageWrapper::FBmpImageWrapper(bool bInHasHeader, bool bInHalfHeight)
 	, bHasHeader(bInHasHeader)
 	, bHalfHeight(bInHalfHeight)
 {
+	// IcoImageWrapper uses FBmpImageWrapper(false, true));
 }
 
 void FBmpImageWrapper::Uncompress(const ERGBFormat InFormat, const int32 InBitDepth)
@@ -30,74 +40,186 @@ void FBmpImageWrapper::Uncompress(const ERGBFormat InFormat, const int32 InBitDe
 
 	if (!bHasHeader || ((CompressedData.Num()>=sizeof(FBitmapFileHeader)+sizeof(FBitmapInfoHeader)) && Buffer[0]=='B' && Buffer[1]=='M'))
 	{
-		UncompressBMPData(InFormat, InBitDepth);
+		if ( ! UncompressBMPData(InFormat, InBitDepth) )
+		{
+			RawData.Empty();
+			if ( LastError.IsEmpty() )
+			{
+				SetError(TEXT("UncompressBMPData failed"));
+			}
+		}
 	}
 }
 
+static inline bool SafeAdvancePointer(const uint8 *& OutPtr, const uint8 * StartPtr, const uint8 * EndPtr, ptrdiff_t Step)
+{
+	if ( Step < 0 || Step > (EndPtr - StartPtr) )
+	{
+		return false;
+	}
+	OutPtr = StartPtr + Step;
+	return true;
+}
+
 
-void FBmpImageWrapper::UncompressBMPData(const ERGBFormat InFormat, const int32 InBitDepth)
+bool FBmpImageWrapper::UncompressBMPData(const ERGBFormat InFormat, const int32 InBitDepth)
 {
 	// always writes BGRA8 :
 	check( InFormat == ERGBFormat::BGRA );
 	check( InBitDepth == 8 );
-	check( ! CompressedData.IsEmpty() );
 
-	const uint8* Buffer = CompressedData.GetData();
+	// was checked before calling here :
+	check( CompressedData.Num() >= sizeof(FBitmapFileHeader)+sizeof(FBitmapInfoHeader) );
+
+	const uint8* const Buffer = CompressedData.GetData();
+	const uint8* const BufferEnd = Buffer + CompressedData.Num();
+
 	const FBitmapInfoHeader* bmhdr = nullptr;
-	const uint8* Bits = nullptr;
+	ptrdiff_t BitsOffset = 0;
 	EBitmapHeaderVersion HeaderVersion = EBitmapHeaderVersion::BHV_BITMAPINFOHEADER;
 
 	if (bHasHeader)
 	{
 		bmhdr = (FBitmapInfoHeader *)(Buffer + sizeof(FBitmapFileHeader));
-		Bits = Buffer + ((FBitmapFileHeader *)Buffer)->bfOffBits;
-		HeaderVersion = ((FBitmapFileHeader *)Buffer)->GetHeaderVersion();
+
+		const FBitmapFileHeader * bmfh = (const FBitmapFileHeader*) Buffer;
+		BitsOffset = bmfh->bfOffBits;
 	}
 	else
 	{
+		// used for ICO loading
 		bmhdr = (FBitmapInfoHeader *)Buffer;
-		Bits = Buffer + sizeof(FBitmapInfoHeader);
+		BitsOffset = bmhdr->biSize;
+	}
+	
+	HeaderVersion = bmhdr->GetHeaderVersion();
+	if ( HeaderVersion == EBitmapHeaderVersion::BHV_INVALID )
+	{
+		UE_LOG(LogImageWrapper, Error, TEXT("BitmapHeaderVersion invalid"));
+		return false;
 	}
 
 	if (bmhdr->biCompression != BCBI_RGB && bmhdr->biCompression != BCBI_BITFIELDS)
 	{
 		UE_LOG(LogImageWrapper, Error, TEXT("RLE compression of BMP images not supported"));
-		return;
+		return false;
+	}
+	
+	if (bmhdr->biPlanes != 1 )
+	{
+		UE_LOG(LogImageWrapper, Error, TEXT("BMP uses an unsupported biPlanes != 1"));
+		return false;
+	}
+	
+	if ( bmhdr->biBitCount == 16 )
+	{
+		UE_LOG(LogImageWrapper, Error, TEXT("BMP 16 bit format no longer supported. Use terrain tools for importing/exporting heightmaps."));
+		return false;
 	}
 
+	if ( bmhdr->biBitCount != 8 && bmhdr->biBitCount != 24 && bmhdr->biBitCount != 32 )
+	{
+		UE_LOG(LogImageWrapper, Error, TEXT("BMP uses an unsupported biBitCount (%i)"), bmhdr->biBitCount);
+		return false;
+	}
+	
+	const uint8* Bits;
+	if ( ! SafeAdvancePointer(Bits,Buffer,BufferEnd,BitsOffset) )
+	{
+		UE_LOG(LogImageWrapper, Error, TEXT("Bmp read would overrun buffer"));
+		return false;
+	}
+	
+	const uint8 * AfterHeader;
+	if ( ! SafeAdvancePointer(AfterHeader, (const uint8 *)bmhdr, BufferEnd, bmhdr->biSize) )
+	{
+		UE_LOG(LogImageWrapper, Error, TEXT("Bmp read would overrun buffer"));
+		return false;
+	}
+
+	// Set texture properties.
+	// This should have already been set by LoadHeader from SetCompressed
+
+	check( Format == ERGBFormat::BGRA );
+	check( BitDepth == 8 );
+
+	Width = bmhdr->biWidth;
+	const bool bNegativeHeight = (bmhdr->biHeight < 0);
+	Height = FMath::Abs(bmhdr->biHeight);
+	if ( bHalfHeight )
+	{
+		Height /= 2;
+	}
+	if ( ! BmpDimensionIsValid(Width) || ! BmpDimensionIsValid(Height) )
+	{
+		UE_LOG(LogImageWrapper, Error, TEXT("Bmp dimensions invalid"));
+		return false;
+	}
+
+	int64 RawDataBytes = Width * (int64)Height * 4;
+	
+	if ( RawDataBytes > INT32_MAX )
+	{
+		UE_LOG(LogImageWrapper, Error, TEXT("Bmp dimensions invalid"));
+		return false;
+	}
+
+	RawData.Empty(RawDataBytes);
+	RawData.AddUninitialized(RawDataBytes);
+	
+	// Copy scanlines, accounting for scanline direction according to the Height field.
+	const int32 SrcBytesPerPel = (bmhdr->biBitCount/8);
+	const int32 SrcStride = Align(Width*SrcBytesPerPel, 4);
+	if ( SrcStride <= 0 )
+	{
+		UE_LOG(LogImageWrapper, Error, TEXT("Bmp dimensions invalid"));
+		return false;
+	}
+
+	const int64 SrcDataSize = SrcStride * (int64) Height;
+	if ( SrcDataSize > (BufferEnd - Bits) )
+	{
+		UE_LOG(LogImageWrapper, Error, TEXT("Bmp read would overrun buffer"));
+		return false;
+	}
+	
+	const int32 SrcPtrDiff = bNegativeHeight ? SrcStride : -SrcStride;
+	const uint8* SrcPtr = Bits + (bNegativeHeight ? 0 : Height - 1) * SrcStride;
+
 	if (bmhdr->biPlanes==1 && bmhdr->biBitCount==8)
 	{
 		// Do palette.
-		const uint8* bmpal = (uint8*)CompressedData.GetData() + sizeof(FBitmapFileHeader) + sizeof(FBitmapInfoHeader);
 
-		// Set texture properties.
-		Width = bmhdr->biWidth;
-		const bool bNegativeHeight = (bmhdr->biHeight < 0);
-		Height = FMath::Abs(bHalfHeight ? bmhdr->biHeight / 2 : bmhdr->biHeight);
-		Format = ERGBFormat::BGRA;
-		RawData.Empty(Height * Width * 4);
-		RawData.AddUninitialized(Height * Width * 4);
+		// If the number for color palette entries is 0, we need to default to 2^biBitCount entries.  In this case 2^8 = 256
+		uint32 clrPaletteCount = bmhdr->biClrUsed ? bmhdr->biClrUsed : 256;
+		if ( clrPaletteCount > 256 )
+		{
+			UE_LOG(LogImageWrapper, Error, TEXT("Bmp paletteCount over 256"));
+			return false;
+		}
+		
+		const uint8* bmpal = AfterHeader;
 
-		FColor* ImageData = (FColor*)RawData.GetData();
+		if ( clrPaletteCount*4 > (BufferEnd - bmpal) )
+		{
+			UE_LOG(LogImageWrapper, Error, TEXT("Bmp read would overrun buffer"));
+			return false;
+		}
 
-		// If the number for color palette entries is 0, we need to default to 2^biBitCount entries.  In this case 2^8 = 256
-		int32 clrPaletteCount = bmhdr->biClrUsed ? bmhdr->biClrUsed : 256;
 		TArray<FColor>	Palette;
+		Palette.SetNum(256);
 
-		for (int32 i = 0; i < clrPaletteCount; i++)
+		for (uint32 i = 0; i < clrPaletteCount; i++)
 		{
-			Palette.Add(FColor(bmpal[i * 4 + 2], bmpal[i * 4 + 1], bmpal[i * 4 + 0], 255));
+			Palette[i] = FColor(bmpal[i * 4 + 2], bmpal[i * 4 + 1], bmpal[i * 4 + 0], 255);
 		}
 
-		while (Palette.Num() < 256)
+		for(uint32 i= clrPaletteCount; i < 256; i++)
 		{
-			Palette.Add(FColor(0, 0, 0, 255));
+			Palette[i] = FColor(0, 0, 0, 255);
 		}
-
-		// Copy scanlines, accounting for scanline direction according to the Height field.
-		const int32 SrcStride = Align(Width, 4);
-		const int32 SrcPtrDiff = bNegativeHeight ? SrcStride : -SrcStride;
-		const uint8* SrcPtr = Bits + (bNegativeHeight ? 0 : Height - 1) * SrcStride;
+		
+		FColor* ImageData = (FColor*)RawData.GetData();
 
 		for (int32 Y = 0; Y < Height; Y++)
 		{
@@ -111,21 +233,8 @@ void FBmpImageWrapper::UncompressBMPData(const ERGBFormat InFormat, const int32
 	}
 	else if (bmhdr->biPlanes==1 && bmhdr->biBitCount==24)
 	{
-		// Set texture properties.
-		Width = bmhdr->biWidth;
-		const bool bNegativeHeight = (bmhdr->biHeight < 0);
-		Height = FMath::Abs(bHalfHeight ? bmhdr->biHeight / 2 : bmhdr->biHeight);
-		Format = ERGBFormat::BGRA;
-		RawData.Empty(Height * Width * 4);
-		RawData.AddUninitialized(Height * Width * 4);
-
 		uint8* ImageData = RawData.GetData();
 
-		// Copy scanlines, accounting for scanline direction according to the Height field.
-		const int32 SrcStride = Align(Width * 3, 4);
-		const int32 SrcPtrDiff = bNegativeHeight ? SrcStride : -SrcStride;
-		const uint8* SrcPtr = Bits + (bNegativeHeight ? 0 : Height - 1) * SrcStride;
-
 		for (int32 Y = 0; Y < Height; Y++)
 		{
 			const uint8* SrcRowPtr = SrcPtr;
@@ -142,28 +251,9 @@ void FBmpImageWrapper::UncompressBMPData(const ERGBFormat InFormat, const int32
 	}
 	else if (bmhdr->biPlanes==1 && bmhdr->biBitCount==32)
 	{
-		// Set texture properties.
-		Width = bmhdr->biWidth;
-		const bool bNegativeHeight = (bmhdr->biHeight < 0);
-		Height = FMath::Abs(bHalfHeight ? bmhdr->biHeight / 2 : bmhdr->biHeight);
-		Format = ERGBFormat::BGRA;
-		RawData.Empty(Height * Width * 4);
-		RawData.AddUninitialized(Height * Width * 4);
-
 		uint8* ImageData = RawData.GetData();
 
-		// Copy scanlines, accounting for scanline direction according to the Height field.
-		const int32 SrcStride = Width * 4;
-		const int32 SrcPtrDiff = bNegativeHeight ? SrcStride : -SrcStride;
-		const uint8* SrcPtr = Bits + (bNegativeHeight ? 0 : Height - 1) * SrcStride;
-
-		// Getting the bmiColors member from the BITMAPINFO, which is used as a mask on BitFields compression.
-		const FBmiColorsMask* ColorMask = (FBmiColorsMask*)(CompressedData.GetData() + sizeof(FBitmapFileHeader) + sizeof(FBitmapInfoHeader));
-		// Header version 4 introduced the option to declare custom color space, so we can't just assume sRGB past that version.
-		const bool bAssumeRGBCompression = bmhdr->biCompression == BCBI_RGB
-			|| (bmhdr->biCompression == BCBI_BITFIELDS && ColorMask->IsMaskRGB8() && HeaderVersion < EBitmapHeaderVersion::BHV_BITMAPV4HEADER);
-		
-		if (bAssumeRGBCompression)
+		if (bmhdr->biCompression == BCBI_RGB)
 		{
 			for (int32 Y = 0; Y < Height; Y++)
 			{
@@ -180,46 +270,86 @@ void FBmpImageWrapper::UncompressBMPData(const ERGBFormat InFormat, const int32
 				SrcPtr += SrcPtrDiff;
 			}
 		}
-		else if (bmhdr->biCompression == BCBI_BITFIELDS)
+		else if (bmhdr->biCompression == BCBI_BITFIELDS) // @todo Oodle : and ALPHABITFIELDS ?
 		{
+			// Advance past the 40-byte header to get to the color masks :
+			//	(note that Adobe wrote the 52 or 56 byte header with biSize=40 so you cannot check biSize to
+			//	 verify you have valid masks)
+			const uint8 * bmhdrEnd;
+			if ( ! SafeAdvancePointer(bmhdrEnd, (const uint8 *)bmhdr, BufferEnd, sizeof(FBitmapInfoHeader)) )
+			{
+				UE_LOG(LogImageWrapper, Error, TEXT("Bmp read would overrun buffer"));
+				return false;
+			}
+
+			// A 52 or 56 byte InfoHeader has masks after a BitmapInfoHeader
+			//	a v4 info header also has them in the same place
+			//	so reading them from there works in both cases
+			const FBmiColorsMask* ColorMask = (FBmiColorsMask*)bmhdrEnd;
+			if ( sizeof(FBmiColorsMask) > (BufferEnd - (const uint8 *)ColorMask) )
+			{
+				UE_LOG(LogImageWrapper, Error, TEXT("Bmp read would overrun buffer"));
+				return false;
+			}
+
+			// Header version 4 introduced the option to declare custom color space, so we can't just assume sRGB past that version.
 			//If the header version is V4 or higher we need to make sure we are still using sRGB format
 			if (HeaderVersion >= EBitmapHeaderVersion::BHV_BITMAPV4HEADER)
 			{
-				const FBitmapInfoHeaderV4* bmhdrV4 = (FBitmapInfoHeaderV4*)(Buffer + sizeof(FBitmapFileHeader));
+				const FBitmapInfoHeaderV4* bmhdrV4 = (FBitmapInfoHeaderV4*)bmhdr;
+				if ( sizeof(FBitmapInfoHeaderV4) > (BufferEnd - (const uint8 *)bmhdrV4) )
+				{
+					UE_LOG(LogImageWrapper, Error, TEXT("Bmp read would overrun buffer"));
+					return false;
+				}
 				
 				if (bmhdrV4->biCSType != (uint32)EBitmapCSType::BCST_LCS_sRGB && bmhdrV4->biCSType != (uint32)EBitmapCSType::BCST_LCS_WINDOWS_COLOR_SPACE)
 				{
-					UE_LOG(LogImageWrapper, Error, TEXT("BMP uses an unsupported custom color space definition, sRGB color space will be used instead."));
+					UE_LOG(LogImageWrapper, Warning, TEXT("BMP uses an unsupported custom color space definition, sRGB color space will be used instead."));
 				}
 			}
 
 			//Calculating the bit mask info needed to remap the pixels' color values.
+			uint32 RGBAMask[4];
 			uint32 TrailingBits[4];
 			float MappingRatio[4];
 			for (uint32 MaskIndex = 0; MaskIndex < 4; MaskIndex++)
 			{
-				TrailingBits[MaskIndex] = FMath::CountTrailingZeros(ColorMask->RGBAMask[MaskIndex]);
-				const uint32 NumberOfBits = 32 - (TrailingBits[MaskIndex] + FMath::CountLeadingZeros(ColorMask->RGBAMask[MaskIndex]));
-				MappingRatio[MaskIndex] = NumberOfBits == 0 ? 0 : (FMath::Exp2(8.f) - 1) / (FMath::Exp2(static_cast<float>(NumberOfBits)) - 1);
+				uint32 Mask = RGBAMask[MaskIndex] = ColorMask->RGBAMask[MaskIndex];
+				if ( Mask == 0 )
+				{
+					TrailingBits[MaskIndex] = 0;
+					MappingRatio[MaskIndex] = 0;
+				}
+				else
+				{
+					// count the number of bits on in Mask by counting the zeros on each side:
+					TrailingBits[MaskIndex] = FMath::CountTrailingZeros(Mask);
+					const uint32 NumberOfBits = 32 - (TrailingBits[MaskIndex] + FMath::CountLeadingZeros(Mask));
+					// note: when NumberOfBits is >= 18, this is not exact (differs from if done in doubles)
+					//		but we still get output in [0,255] and the loss is small, so let it be
+					MappingRatio[MaskIndex] = NumberOfBits == 0 ? 0 : (255.f / ((1ULL<<NumberOfBits) - 1) );
+				}
 			}
 
 			//In header pre-version 4, we should ignore the last 32bit (alpha) content.
-			const bool bHasAlphaChannel = ColorMask->RGBAMask[3] != 0 && HeaderVersion >= EBitmapHeaderVersion::BHV_BITMAPV4HEADER;
+			// @todo Oodle: 56-byte Adobe headers might also have valid alpha masks
+			const bool bHasAlphaChannel = RGBAMask[3] != 0 && HeaderVersion >= EBitmapHeaderVersion::BHV_BITMAPV4HEADER;
 
 			for (int32 Y = 0; Y < Height; Y++)
 			{
-				const uint32* SrcPixel = (uint32*)SrcPtr;
 				for (int32 X = 0; X < Width; X++)
 				{
-					// Set the color values in BGRA order.
-					for (int32 ColorIndex = 2; ColorIndex >= 0; ColorIndex--)
-					{
-						*ImageData++ = FMath::RoundToInt(((*SrcPixel & ColorMask->RGBAMask[ColorIndex]) >> TrailingBits[ColorIndex]) * MappingRatio[ColorIndex]);
-					}
+					const uint32 SrcPixel = ((const uint32*)SrcPtr)[X];
 
-					*ImageData++ = bHasAlphaChannel ? FMath::RoundToInt(((*SrcPixel & ColorMask->RGBAMask[3]) >> TrailingBits[3]) * MappingRatio[3]) : 0xFF;
+					// Set the color values in BGRA order.
+					//	integer output of RoundToInt will always fit in U8, no clamp needed
+					// @todo Oodle: the shift by TrailingBits can be combined into the float multiply using ldexpf(MappingRatio,-TrailingBits)
+					*ImageData++ = (uint8) (0.5f + ((SrcPixel & RGBAMask[2]) >> TrailingBits[2]) * MappingRatio[2]);
+					*ImageData++ = (uint8) (0.5f + ((SrcPixel & RGBAMask[1]) >> TrailingBits[1]) * MappingRatio[1]);
+					*ImageData++ = (uint8) (0.5f + ((SrcPixel & RGBAMask[0]) >> TrailingBits[0]) * MappingRatio[0]);
 
-					SrcPixel++;
+					*ImageData++ = bHasAlphaChannel ? (uint8) (0.5f + ((SrcPixel & RGBAMask[3]) >> TrailingBits[3]) * MappingRatio[3]) : 0xFF;
 				}
 
 				SrcPtr += SrcPtrDiff;
@@ -228,16 +358,21 @@ void FBmpImageWrapper::UncompressBMPData(const ERGBFormat InFormat, const int32
 		else
 		{
 			UE_LOG(LogImageWrapper, Error, TEXT("BMP uses an unsupported compression format (%i)"), bmhdr->biCompression)
+			return false;
 		}
 	}
 	else if (bmhdr->biPlanes==1 && bmhdr->biBitCount==16)
 	{
 		UE_LOG(LogImageWrapper, Error, TEXT("BMP 16 bit format no longer supported. Use terrain tools for importing/exporting heightmaps."));
+		return false;
 	}
 	else
 	{
-		UE_LOG(LogImageWrapper, Error, TEXT("BMP uses an unsupported format (%i/%i)"), bmhdr->biPlanes, bmhdr->biBitCount);
+		UE_LOG(LogImageWrapper, Error, TEXT("BMP uses an unsupported format (planes = %i, bitcount = %i)"), bmhdr->biPlanes, bmhdr->biBitCount);
+		return false;
 	}
+
+	return true;
 }
 
 
@@ -245,7 +380,7 @@ bool FBmpImageWrapper::SetCompressed(const void* InCompressedData, int64 InCompr
 {
 	bool bResult = FImageWrapperBase::SetCompressed(InCompressedData, InCompressedSize);
 
-	bResult = bResult && (bHasHeader ? LoadBMPHeader() : LoadBMPInfoHeader());	// Fetch the variables from the header info
+	bResult = bResult && (bHasHeader ? LoadBMPHeader() : LoadBMPInfoHeader(0));	// Fetch the variables from the header info
 
 	if ( ! bResult )
 	{
@@ -258,50 +393,32 @@ bool FBmpImageWrapper::SetCompressed(const void* InCompressedData, int64 InCompr
 
 bool FBmpImageWrapper::LoadBMPHeader()
 {
-	//note: not Endian correct
+	if ( CompressedData.Num() < sizeof(FBitmapInfoHeader) + sizeof(FBitmapFileHeader) )
+	{
+		UE_LOG(LogImageWrapper, Error, TEXT("Bmp read would overrun buffer"));
+		return false;
+	}
+
 	const FBitmapInfoHeader* bmhdr = (FBitmapInfoHeader *)(CompressedData.GetData() + sizeof(FBitmapFileHeader));
-	const FBitmapFileHeader* bmf   = (FBitmapFileHeader *)(CompressedData.GetData() + 0);
 	if ((CompressedData.Num() >= sizeof(FBitmapFileHeader) + sizeof(FBitmapInfoHeader)) && CompressedData.GetData()[0] == 'B' && CompressedData.GetData()[1] == 'M')
 	{
-		if (bmhdr->biCompression != BCBI_RGB && bmhdr->biCompression != BCBI_BITFIELDS)
-		{
-			UE_LOG(LogImageWrapper, Error, TEXT("RLE compression of BMP images not supported"));
-			return false;
-		}
-
-		if (bmhdr->biPlanes==1 && (bmhdr->biBitCount==8 || bmhdr->biBitCount==24 || bmhdr->biBitCount==32))
-		{
-			// Set texture properties.
-			Width = bmhdr->biWidth;
-			Height = FMath::Abs(bmhdr->biHeight);
-			Format = ERGBFormat::BGRA;
-			//  YUCK
-			//  BmpImageWrapper was reporting BitDepth *total* not per-channel (eg. 24)
-			//  some legacy code knew that and expected to see the bad values, beware
-			//BitDepth = bmhdr->biBitCount;
-			BitDepth = 8;
-
-			return true;
-		}
-		
-		if (bmhdr->biPlanes == 1 && bmhdr->biBitCount == 16)
-		{
-			UE_LOG(LogImageWrapper, Error, TEXT("BMP 16 bit format no longer supported. Use terrain tools for importing/exporting heightmaps."));
-		}
-		else
-		{
-			UE_LOG(LogImageWrapper, Error, TEXT("BMP uses an unsupported format (%i/%i)"), bmhdr->biPlanes, bmhdr->biBitCount);
-		}
+		return LoadBMPInfoHeader(sizeof(FBitmapFileHeader));
 	}
-
+	
+	UE_LOG(LogImageWrapper, Error, TEXT("Bmp header invalid"));
 	return false;
 }
 
 
-bool FBmpImageWrapper::LoadBMPInfoHeader()
+bool FBmpImageWrapper::LoadBMPInfoHeader(int64 HeaderOffset)
 {
-	//note: not Endian correct
-	const FBitmapInfoHeader* bmhdr = (FBitmapInfoHeader *)CompressedData.GetData();
+	if ( CompressedData.Num() < HeaderOffset+(int64)sizeof(FBitmapInfoHeader) )
+	{
+		UE_LOG(LogImageWrapper, Error, TEXT("Bmp read would overrun buffer"));
+		return false;
+	}
+
+	const FBitmapInfoHeader* bmhdr = (FBitmapInfoHeader *)(CompressedData.GetData() + HeaderOffset);
 
 	if (bmhdr->biCompression != BCBI_RGB && bmhdr->biCompression != BCBI_BITFIELDS)
 	{
@@ -314,24 +431,32 @@ bool FBmpImageWrapper::LoadBMPInfoHeader()
 		// Set texture properties.
 		Width = bmhdr->biWidth;
 		Height = FMath::Abs(bmhdr->biHeight);
+		if ( bHalfHeight )
+		{
+			Height /= 2;
+		}
 		Format = ERGBFormat::BGRA;
-		//  BmpImageWrapper was reporting BitDepth *total* not per-channel (eg. 24)
-		//BitDepth = bmhdr->biBitCount;
 		BitDepth = 8;
 
+		if ( ! BmpDimensionIsValid(Width) || ! BmpDimensionIsValid(Height) )
+		{
+			UE_LOG(LogImageWrapper, Error, TEXT("Bmp dimensions invalid"));
+			return false;
+		}
+
 		return true;
 	}
 	
 	if (bmhdr->biPlanes == 1 && bmhdr->biBitCount == 16)
 	{
 		UE_LOG(LogImageWrapper, Error, TEXT("BMP 16 bit format no longer supported. Use terrain tools for importing/exporting heightmaps."));
+		return false;
 	}
 	else
 	{
-		UE_LOG(LogImageWrapper, Error, TEXT("BMP uses an unsupported format (%i/%i)"), bmhdr->biPlanes, bmhdr->biBitCount);
+		UE_LOG(LogImageWrapper, Error, TEXT("BMP uses an unsupported format (planes = %i, bitcount = %i)"), bmhdr->biPlanes, bmhdr->biBitCount);
+		return false;
 	}
-
-	return false;
 }
 
 bool FBmpImageWrapper::CanSetRawFormat(const ERGBFormat InFormat, const int32 InBitDepth) const
@@ -367,9 +492,12 @@ void FBmpImageWrapper::Compress(int32 Quality)
 	check( Format == ERGBFormat::BGRA || Format == ERGBFormat::Gray );
 	check( BitDepth == 8 );
 
+	check( BmpDimensionIsValid(Width) );
+	check( BmpDimensionIsValid(Height) );
+
 	// write 8,24, or 32 bit bmp
 
-	int64 NumPixels = Width*Height;
+	int64 NumPixels = Width*(int64)Height;
 	int64 RawDataSize = RawData.Num();
 
 	int RawBytesPerPel = (Format == ERGBFormat::BGRA) ? 4 : 1;
diff --git a/Engine/Source/Runtime/ImageWrapper/Private/Formats/BmpImageWrapper.h b/Engine/Source/Runtime/ImageWrapper/Private/Formats/BmpImageWrapper.h
--- a/Engine/Source/Runtime/ImageWrapper/Private/Formats/BmpImageWrapper.h
+++ b/Engine/Source/Runtime/ImageWrapper/Private/Formats/BmpImageWrapper.h
@@ -20,7 +20,7 @@ public:
 public:
 
 	/** Helper function used to uncompress BMP data from a buffer */
-	void UncompressBMPData(const ERGBFormat InFormat, const int32 InBitDepth);
+	bool UncompressBMPData(const ERGBFormat InFormat, const int32 InBitDepth);
 
 	/** 
 	 * Load the header information, returns true if successful.
@@ -34,7 +34,7 @@ public:
 	 *
 	 * @return true if successful
 	 */
-	bool LoadBMPInfoHeader();
+	bool LoadBMPInfoHeader(int64 HeaderOffset);
 
 public:
 
diff --git a/Engine/Source/Runtime/ImageWrapper/Private/Formats/IcoImageWrapper.cpp b/Engine/Source/Runtime/ImageWrapper/Private/Formats/IcoImageWrapper.cpp
--- a/Engine/Source/Runtime/ImageWrapper/Private/Formats/IcoImageWrapper.cpp
+++ b/Engine/Source/Runtime/ImageWrapper/Private/Formats/IcoImageWrapper.cpp
@@ -93,6 +93,9 @@ void FIcoImageWrapper::Uncompress( const ERGBFormat InFormat, const int32 InBitD
 	if (ImageOffset != 0 && ImageSize != 0)
 	{
 		SubImageWrapper->Uncompress(InFormat, InBitDepth);
+		// Uncompress has no return value
+		//  we can tell it failed if it set an error, or has no rawdata
+		LastError = SubImageWrapper->GetLastError();
 	}
 }
 
@@ -113,6 +116,10 @@ bool FIcoImageWrapper::GetRaw( const ERGBFormat InFormat, int32 InBitDepth, TArr
 	if (LastError.IsEmpty())
 	{
 		SubImageWrapper->MoveRawData(OutRawData);
+		if ( OutRawData.IsEmpty() )
+		{
+			return false;
+		}
 	}
 
 	return LastError.IsEmpty();
@@ -164,7 +171,7 @@ bool FIcoImageWrapper::LoadICOHeader()
 				{
 					// otherwise this should be a BMP icon
 					Width = BmpWrapper->GetWidth();
-					Height = BmpWrapper->GetHeight() / 2;	// ICO file spec says to divide by 2 here as height refers to combined image & mask height
+					Height = BmpWrapper->GetHeight();
 					Format = BmpWrapper->GetFormat();
 					BitDepth = BmpWrapper->GetBitDepth();
 					LargestWidth = RealWidth;
diff --git a/Engine/Source/Runtime/ImageWrapper/Private/ImageWrapperBase.h b/Engine/Source/Runtime/ImageWrapper/Private/ImageWrapperBase.h
--- a/Engine/Source/Runtime/ImageWrapper/Private/ImageWrapperBase.h
+++ b/Engine/Source/Runtime/ImageWrapper/Private/ImageWrapperBase.h
@@ -64,6 +64,14 @@ public:
 	 * @param ErrorMessage The error message to set.
 	 */
 	void SetError(const TCHAR* ErrorMessage);
+	
+	/**
+	 * Gets last error message.
+	 */
+	const FString & GetLastError() const
+	{
+		return LastError;
+	}
 
 	/**  
 	 * Function to uncompress our data 
diff --git a/Engine/Source/Runtime/ImageWrapper/Public/BmpImageSupport.h b/Engine/Source/Runtime/ImageWrapper/Public/BmpImageSupport.h
--- a/Engine/Source/Runtime/ImageWrapper/Public/BmpImageSupport.h
+++ b/Engine/Source/Runtime/ImageWrapper/Public/BmpImageSupport.h
@@ -27,6 +27,7 @@ enum class EBitmapHeaderVersion : uint8
 	BHV_BITMAPV3INFOHEADER = 2,
 	BHV_BITMAPV4HEADER     = 3,
 	BHV_BITMAPV5HEADER     = 4,
+	BHV_INVALID = 0xFF
 };
 
 // Color space type of the bitmap, property introduced in Bitmap header version 4.
@@ -54,27 +55,6 @@ struct FBitmapFileHeader
 		Ar << H.bfType << H.bfSize << H.bfReserved1 << H.bfReserved2 << H.bfOffBits;
 		return Ar;
 	}
-
-public:
-	EBitmapHeaderVersion GetHeaderVersion() const
-	{
-		// Since there is no field indicating the header version of the bitmap in the FileHeader,
-		// the only way to know the format version is to check the header offbits size.
-		switch (bfOffBits - sizeof(FBitmapFileHeader))
-		{
-			case 40:
-			default:
-				return EBitmapHeaderVersion::BHV_BITMAPINFOHEADER;
-			case 52:
-				return EBitmapHeaderVersion::BHV_BITMAPV2INFOHEADER;
-			case 56:
-				return EBitmapHeaderVersion::BHV_BITMAPV3INFOHEADER;
-			case 108:
-				return EBitmapHeaderVersion::BHV_BITMAPV4HEADER;
-			case 124:
-				return EBitmapHeaderVersion::BHV_BITMAPV5HEADER;
-		}
-	}
 };
 #pragma pack(pop)
 
@@ -102,6 +82,33 @@ struct FBitmapInfoHeader
 		Ar << H.biClrUsed << H.biClrImportant;
 		return Ar;
 	}
+	
+public:
+	EBitmapHeaderVersion GetHeaderVersion() const
+	{
+		// Since there is no field indicating the header version of the bitmap in the FileHeader,
+		// the only way to know the format version is to check the header size.
+		//
+		// note that Adobe (incorrectly) writes biSize as 40 and then BMFH bfOffBits will be 52 or 56
+		//	so this code does not detect those variant headers
+		switch (biSize)
+		{
+			case 40:
+				return EBitmapHeaderVersion::BHV_BITMAPINFOHEADER;
+			case 52:
+				// + RGB 32 bit masks
+				return EBitmapHeaderVersion::BHV_BITMAPV2INFOHEADER;
+			case 56:
+				// + RGBA 32 bit masks
+				return EBitmapHeaderVersion::BHV_BITMAPV3INFOHEADER;
+			case 108:
+				return EBitmapHeaderVersion::BHV_BITMAPV4HEADER;
+			case 124:
+				return EBitmapHeaderVersion::BHV_BITMAPV5HEADER;
+			default:
+				return EBitmapHeaderVersion::BHV_INVALID;
+		}
+	}
 };
 #pragma pack(pop)
 
@@ -163,13 +170,5 @@ struct FBmiColorsMask
 {
 	// RGBA, in header pre-version 4, Alpha was only used as padding.
 	uint32 RGBAMask[4];
-
-public: 
-	bool IsMaskRGB8() const
-	{
-		return FMath::CountLeadingZeros(RGBAMask[0]) + FMath::CountTrailingZeros(RGBAMask[0]) == 24
-			&& FMath::CountLeadingZeros(RGBAMask[1]) + FMath::CountTrailingZeros(RGBAMask[1]) == 24
-			&& FMath::CountLeadingZeros(RGBAMask[2]) + FMath::CountTrailingZeros(RGBAMask[2]) == 24;
-	}
 };
 #pragma pack(pop)
\ No newline at end of file
