BmpImageWrapper support 16 bit bmp like 555 and 565
also import A channel from 32-bit BI_RGB
now loads all formats that photoshop "advanced" outputs

diff --git a/Engine/Source/Runtime/ImageWrapper/Private/Formats/BmpImageWrapper.cpp b/Engine/Source/Runtime/ImageWrapper/Private/Formats/BmpImageWrapper.cpp
--- a/Engine/Source/Runtime/ImageWrapper/Private/Formats/BmpImageWrapper.cpp
+++ b/Engine/Source/Runtime/ImageWrapper/Private/Formats/BmpImageWrapper.cpp
@@ -99,7 +99,13 @@ bool FBmpImageWrapper::UncompressBMPData(const ERGBFormat InFormat, const int32
 		return false;
 	}
 
-	if (bmhdr->biCompression != BCBI_RGB && bmhdr->biCompression != BCBI_BITFIELDS)
+	/*
+	UE_LOG(LogImageWrapper, Log, TEXT("BMP compression = (%i) BitCount = (%i)"), bmhdr->biCompression,bmhdr->biBitCount)
+	UE_LOG(LogImageWrapper, Log, TEXT("BMP BitsOffset = (%i)"), BitsOffset)
+	UE_LOG(LogImageWrapper, Log, TEXT("BMP biSize = (%i)"), bmhdr->biSize)
+	*/
+
+	if (bmhdr->biCompression != BCBI_RGB && bmhdr->biCompression != BCBI_BITFIELDS && bmhdr->biCompression != BCBI_ALPHABITFIELDS)
 	{
 		UE_LOG(LogImageWrapper, Error, TEXT("RLE compression of BMP images not supported"));
 		return false;
@@ -111,13 +117,7 @@ bool FBmpImageWrapper::UncompressBMPData(const ERGBFormat InFormat, const int32
 		return false;
 	}
 	
-	if ( bmhdr->biBitCount == 16 )
-	{
-		UE_LOG(LogImageWrapper, Error, TEXT("BMP 16 bit format no longer supported. Use terrain tools for importing/exporting heightmaps."));
-		return false;
-	}
-
-	if ( bmhdr->biBitCount != 8 && bmhdr->biBitCount != 24 && bmhdr->biBitCount != 32 )
+	if ( bmhdr->biBitCount != 8 && bmhdr->biBitCount != 16 && bmhdr->biBitCount != 24 && bmhdr->biBitCount != 32 )
 	{
 		UE_LOG(LogImageWrapper, Error, TEXT("BMP uses an unsupported biBitCount (%i)"), bmhdr->biBitCount);
 		return false;
@@ -166,9 +166,11 @@ bool FBmpImageWrapper::UncompressBMPData(const ERGBFormat InFormat, const int32
 
 	RawData.Empty(RawDataBytes);
 	RawData.AddUninitialized(RawDataBytes);
+	uint8* ImageData = RawData.GetData();
 	
 	// Copy scanlines, accounting for scanline direction according to the Height field.
 	const int32 SrcBytesPerPel = (bmhdr->biBitCount/8);
+	check( SrcBytesPerPel*8 == bmhdr->biBitCount );
 	const int32 SrcStride = Align(Width*SrcBytesPerPel, 4);
 	if ( SrcStride <= 0 )
 	{
@@ -186,7 +188,7 @@ bool FBmpImageWrapper::UncompressBMPData(const ERGBFormat InFormat, const int32
 	const int32 SrcPtrDiff = bNegativeHeight ? SrcStride : -SrcStride;
 	const uint8* SrcPtr = Bits + (bNegativeHeight ? 0 : Height - 1) * SrcStride;
 
-	if (bmhdr->biPlanes==1 && bmhdr->biBitCount==8)
+	if ( bmhdr->biBitCount==8)
 	{
 		// Do palette.
 
@@ -219,22 +221,20 @@ bool FBmpImageWrapper::UncompressBMPData(const ERGBFormat InFormat, const int32
 			Palette[i] = FColor(0, 0, 0, 255);
 		}
 		
-		FColor* ImageData = (FColor*)RawData.GetData();
+		FColor* ImageColors = (FColor*)ImageData;
 
 		for (int32 Y = 0; Y < Height; Y++)
 		{
 			for (int32 X = 0; X < Width; X++)
 			{
-				*ImageData++ = Palette[SrcPtr[X]];
+				*ImageColors++ = Palette[SrcPtr[X]];
 			}
 
 			SrcPtr += SrcPtrDiff;
 		}
 	}
-	else if (bmhdr->biPlanes==1 && bmhdr->biBitCount==24)
+	else if ( bmhdr->biBitCount==24)
 	{
-		uint8* ImageData = RawData.GetData();
-
 		for (int32 Y = 0; Y < Height; Y++)
 		{
 			const uint8* SrcRowPtr = SrcPtr;
@@ -249,93 +249,163 @@ bool FBmpImageWrapper::UncompressBMPData(const ERGBFormat InFormat, const int32
 			SrcPtr += SrcPtrDiff;
 		}
 	}
-	else if (bmhdr->biPlanes==1 && bmhdr->biBitCount==32)
+	else if ( bmhdr->biBitCount==32 && bmhdr->biCompression == BCBI_RGB)
 	{
-		uint8* ImageData = RawData.GetData();
+		// This comment was previously here :
+		//  "In BCBI_RGB compression the last 8 bits of the pixel are not used."
+		// -> this agrees with MSDN but does not match what Photoshop does in practice
+		//  photoshop writes 32-bit ARGB with non-trivial A using BI_RGB
+		//	see "porsche512a_ARGB.bmp" also "porsche512a_notadvanced.bmp"
+		//	(both the "advanced" and regular photoshop save do this)
+		// 
+
+		uint64 TotalA = 0;
+		for (int32 Y = 0; Y < Height; Y++)
+		{
+			// this is just a memcpy, except the accumulation of TotalA
+			const uint8* SrcRowPtr = SrcPtr;
+			for (int32 X = 0; X < Width; X++)
+			{
+				*ImageData++ = *SrcRowPtr++;
+				*ImageData++ = *SrcRowPtr++;
+				*ImageData++ = *SrcRowPtr++;
+				TotalA += *SrcRowPtr;
+				*ImageData++ = *SrcRowPtr++; // was doing = 0xFF , ignoring fourth byte, as per MSDN
+			}
+
+			SrcPtr += SrcPtrDiff;
+		}
 
-		if (bmhdr->biCompression == BCBI_RGB)
+		if ( TotalA == 0 )
 		{
+			// assume that this is actually XRGB and they wrote zeros in A
+			// go back through and change all A's to 0xFF
+			
+			ImageData = RawData.GetData();
+	
 			for (int32 Y = 0; Y < Height; Y++)
 			{
-				const uint8* SrcRowPtr = SrcPtr;
 				for (int32 X = 0; X < Width; X++)
 				{
-					*ImageData++ = *SrcRowPtr++;
-					*ImageData++ = *SrcRowPtr++;
-					*ImageData++ = *SrcRowPtr++;
-					*ImageData++ = 0xFF; //In BCBI_RGB compression the last 8 bits of the pixel are not used.
-					SrcRowPtr++;
+					ImageData[3] = 0xFF;
+					ImageData +=  4;
 				}
-
-				SrcPtr += SrcPtrDiff;
 			}
 		}
-		else if (bmhdr->biCompression == BCBI_BITFIELDS) // @todo Oodle : and ALPHABITFIELDS ?
+	}
+	else if ( bmhdr->biBitCount==16 && bmhdr->biCompression == BCBI_RGB)
+	{
+		// 16 bit BI_RGB is 555
+
+		for (int32 Y = 0; Y < Height; Y++)
 		{
-			// Advance past the 40-byte header to get to the color masks :
-			//	(note that Adobe wrote the 52 or 56 byte header with biSize=40 so you cannot check biSize to
-			//	 verify you have valid masks)
-			const uint8 * bmhdrEnd;
-			if ( ! SafeAdvancePointer(bmhdrEnd, (const uint8 *)bmhdr, BufferEnd, sizeof(FBitmapInfoHeader)) )
+			for (int32 X = 0; X < Width; X++)
 			{
-				UE_LOG(LogImageWrapper, Error, TEXT("Bmp read would overrun buffer"));
-				return false;
+				const uint32 SrcPixel = ((const uint16*)SrcPtr)[X];
+
+				// Set the color values in BGRA order.
+
+				uint32 r = (SrcPixel>>10)&0x1f;;
+				uint32 g = (SrcPixel>> 5)&0x1f;
+				uint32 b = (SrcPixel    )&0x1f;
+
+				*ImageData++ = (uint8) ( (b<<3) | (b>>2) );
+				*ImageData++ = (uint8) ( (g<<3) | (g>>2) );
+				*ImageData++ = (uint8) ( (r<<3) | (r>>2) );
+				*ImageData++ = 0xFF; // 555 BI_RGB does not use alpha bit
 			}
 
-			// A 52 or 56 byte InfoHeader has masks after a BitmapInfoHeader
-			//	a v4 info header also has them in the same place
-			//	so reading them from there works in both cases
-			const FBmiColorsMask* ColorMask = (FBmiColorsMask*)bmhdrEnd;
-			if ( sizeof(FBmiColorsMask) > (BufferEnd - (const uint8 *)ColorMask) )
+			SrcPtr += SrcPtrDiff;
+		}
+	}
+	else if ( ( bmhdr->biBitCount==16 || bmhdr->biBitCount==32 ) && ( bmhdr->biCompression == BCBI_BITFIELDS || bmhdr->biCompression == BCBI_ALPHABITFIELDS) )
+	{
+		// Advance past the 40-byte header to get to the color masks :
+		//	(note that some bmps have the 52 or 56 byte header with biSize=40 so you cannot check biSize to verify you have valid masks)
+		//check( bmhdr->biSize >= 52 );
+		//	in theory you could check BitsOffset
+
+		const uint8 * bmhdrEnd;
+		if ( ! SafeAdvancePointer(bmhdrEnd, (const uint8 *)bmhdr, BufferEnd, sizeof(FBitmapInfoHeader)) )
+		{
+			UE_LOG(LogImageWrapper, Error, TEXT("Bmp read would overrun buffer"));
+			return false;
+		}
+
+		// A 52 or 56 byte InfoHeader has masks after a BitmapInfoHeader
+		//	a v4 info header also has them in the same place
+		//	so reading them from there works in both cases
+		const FBmiColorsMask* ColorMask = (FBmiColorsMask*)bmhdrEnd;
+		if ( sizeof(FBmiColorsMask) > (BufferEnd - (const uint8 *)ColorMask) )
+		{
+			UE_LOG(LogImageWrapper, Error, TEXT("Bmp read would overrun buffer"));
+			return false;
+		}
+
+		// Header version 4 introduced the option to declare custom color space, so we can't just assume sRGB past that version.
+		//If the header version is V4 or higher we need to make sure we are still using sRGB format
+		if (HeaderVersion >= EBitmapHeaderVersion::BHV_BITMAPV4HEADER)
+		{
+			const FBitmapInfoHeaderV4* bmhdrV4 = (FBitmapInfoHeaderV4*)bmhdr;
+			if ( sizeof(FBitmapInfoHeaderV4) > (BufferEnd - (const uint8 *)bmhdrV4) )
 			{
 				UE_LOG(LogImageWrapper, Error, TEXT("Bmp read would overrun buffer"));
 				return false;
 			}
-
-			// Header version 4 introduced the option to declare custom color space, so we can't just assume sRGB past that version.
-			//If the header version is V4 or higher we need to make sure we are still using sRGB format
-			if (HeaderVersion >= EBitmapHeaderVersion::BHV_BITMAPV4HEADER)
-			{
-				const FBitmapInfoHeaderV4* bmhdrV4 = (FBitmapInfoHeaderV4*)bmhdr;
-				if ( sizeof(FBitmapInfoHeaderV4) > (BufferEnd - (const uint8 *)bmhdrV4) )
-				{
-					UE_LOG(LogImageWrapper, Error, TEXT("Bmp read would overrun buffer"));
-					return false;
-				}
 				
-				if (bmhdrV4->biCSType != (uint32)EBitmapCSType::BCST_LCS_sRGB && bmhdrV4->biCSType != (uint32)EBitmapCSType::BCST_LCS_WINDOWS_COLOR_SPACE)
-				{
-					UE_LOG(LogImageWrapper, Warning, TEXT("BMP uses an unsupported custom color space definition, sRGB color space will be used instead."));
-				}
+			if (bmhdrV4->biCSType != (uint32)EBitmapCSType::BCST_LCS_sRGB && bmhdrV4->biCSType != (uint32)EBitmapCSType::BCST_LCS_WINDOWS_COLOR_SPACE)
+			{
+				UE_LOG(LogImageWrapper, Warning, TEXT("BMP uses an unsupported custom color space definition, sRGB color space will be used instead."));
 			}
+		}
 
-			//Calculating the bit mask info needed to remap the pixels' color values.
-			uint32 RGBAMask[4];
-			uint32 TrailingBits[4];
-			float MappingRatio[4];
-			for (uint32 MaskIndex = 0; MaskIndex < 4; MaskIndex++)
+		//Calculating the bit mask info needed to remap the pixels' color values.
+		// note RGBAMask[3] can be reading past the end of the header
+		//	but we will just read payload bits, and then replace it later when !bHasAlphaChannel
+		uint32 RGBAMask[4];
+		float MappingRatio[4];
+		for (uint32 MaskIndex = 0; MaskIndex < 4; MaskIndex++)
+		{
+			uint32 Mask = RGBAMask[MaskIndex] = ColorMask->RGBAMask[MaskIndex];
+			if ( Mask == 0 )
 			{
-				uint32 Mask = RGBAMask[MaskIndex] = ColorMask->RGBAMask[MaskIndex];
-				if ( Mask == 0 )
-				{
-					TrailingBits[MaskIndex] = 0;
-					MappingRatio[MaskIndex] = 0;
-				}
-				else
-				{
-					// count the number of bits on in Mask by counting the zeros on each side:
-					TrailingBits[MaskIndex] = FMath::CountTrailingZeros(Mask);
-					const uint32 NumberOfBits = 32 - (TrailingBits[MaskIndex] + FMath::CountLeadingZeros(Mask));
-					// note: when NumberOfBits is >= 18, this is not exact (differs from if done in doubles)
-					//		but we still get output in [0,255] and the loss is small, so let it be
-					MappingRatio[MaskIndex] = NumberOfBits == 0 ? 0 : (255.f / ((1ULL<<NumberOfBits) - 1) );
-				}
+				MappingRatio[MaskIndex] = 0;
 			}
+			else
+			{
+				// count the number of bits on in Mask by counting the zeros on each side:
+				int32 TrailingBits = FMath::CountTrailingZeros(Mask);
+				int32 NumberOfBits = 32 - (TrailingBits + FMath::CountLeadingZeros(Mask));
+				check( NumberOfBits > 0 );
+
+				// note: when NumberOfBits is >= 18, this is not exact (differs from if done in doubles)
+				//		but we still get output in [0,255] and the error is small, so let it be
+				// use ldexpf to put the >>TrailingBits in the multiply
+				MappingRatio[MaskIndex] = ldexpf( (255.f / ((1ULL<<NumberOfBits) - 1) ) , -TrailingBits );
+			}
+		}
 
-			//In header pre-version 4, we should ignore the last 32bit (alpha) content.
-			// @todo Oodle: 56-byte Adobe headers might also have valid alpha masks
-			const bool bHasAlphaChannel = RGBAMask[3] != 0 && HeaderVersion >= EBitmapHeaderVersion::BHV_BITMAPV4HEADER;
+		//In header pre-version 4, we should ignore the last 32bit (alpha) content.
+		const bool bHasAlphaChannel = RGBAMask[3] != 0 && HeaderVersion >= EBitmapHeaderVersion::BHV_BITMAPV4HEADER;
+		if ( bmhdr->biSize == 56 && RGBAMask[3] != 0 )
+		{
+			// 56-byte Adobe headers ("bmpv3") might also have valid alpha masks
+			//	legacy Unreal import was treating them as bHasAlphaChannel=false
+			//	perhaps they should in fact have their alpha mask respected
+			// note that Adobe actually uses BI_RGB not BI_BITFIELDS for ARGB output in some cases
+			UE_LOG(LogImageWrapper, Display, TEXT("Adobe 56-byte header might have alpha but we ignore it %08X"), RGBAMask[3]);				
+		}
 
+		float AlphaBias = 0.5f;
+		if ( ! bHasAlphaChannel )
+		{
+			RGBAMask[3] = 0;
+			MappingRatio[3] = 0.f;
+			AlphaBias = 255.f;
+		}
+		
+		if ( bmhdr->biBitCount == 32 )
+		{
 			for (int32 Y = 0; Y < Height; Y++)
 			{
 				for (int32 X = 0; X < Width; X++)
@@ -344,12 +414,10 @@ bool FBmpImageWrapper::UncompressBMPData(const ERGBFormat InFormat, const int32
 
 					// Set the color values in BGRA order.
 					//	integer output of RoundToInt will always fit in U8, no clamp needed
-					// @todo Oodle: the shift by TrailingBits can be combined into the float multiply using ldexpf(MappingRatio,-TrailingBits)
-					*ImageData++ = (uint8) (0.5f + ((SrcPixel & RGBAMask[2]) >> TrailingBits[2]) * MappingRatio[2]);
-					*ImageData++ = (uint8) (0.5f + ((SrcPixel & RGBAMask[1]) >> TrailingBits[1]) * MappingRatio[1]);
-					*ImageData++ = (uint8) (0.5f + ((SrcPixel & RGBAMask[0]) >> TrailingBits[0]) * MappingRatio[0]);
-
-					*ImageData++ = bHasAlphaChannel ? (uint8) (0.5f + ((SrcPixel & RGBAMask[3]) >> TrailingBits[3]) * MappingRatio[3]) : 0xFF;
+					*ImageData++ = (uint8) ((SrcPixel & RGBAMask[2]) * MappingRatio[2] + 0.5f);
+					*ImageData++ = (uint8) ((SrcPixel & RGBAMask[1]) * MappingRatio[1] + 0.5f);
+					*ImageData++ = (uint8) ((SrcPixel & RGBAMask[0]) * MappingRatio[0] + 0.5f);
+					*ImageData++ = (uint8) ((SrcPixel & RGBAMask[3]) * MappingRatio[3] + AlphaBias);
 				}
 
 				SrcPtr += SrcPtrDiff;
@@ -357,18 +425,31 @@ bool FBmpImageWrapper::UncompressBMPData(const ERGBFormat InFormat, const int32
 		}
 		else
 		{
-			UE_LOG(LogImageWrapper, Error, TEXT("BMP uses an unsupported compression format (%i)"), bmhdr->biCompression)
-			return false;
+			// code dupe: only change from above loop is the type cast on SrcPtr[X]
+			check( bmhdr->biBitCount == 16 );
+
+			for (int32 Y = 0; Y < Height; Y++)
+			{
+				for (int32 X = 0; X < Width; X++)
+				{
+					const uint32 SrcPixel = ((const uint16*)SrcPtr)[X];
+
+					// Set the color values in BGRA order.
+					//	integer output of RoundToInt will always fit in U8, no clamp needed
+					*ImageData++ = (uint8) ((SrcPixel & RGBAMask[2]) * MappingRatio[2] + 0.5f);
+					*ImageData++ = (uint8) ((SrcPixel & RGBAMask[1]) * MappingRatio[1] + 0.5f);
+					*ImageData++ = (uint8) ((SrcPixel & RGBAMask[0]) * MappingRatio[0] + 0.5f);
+					*ImageData++ = (uint8) ((SrcPixel & RGBAMask[3]) * MappingRatio[3] + AlphaBias);
+				}
+
+				SrcPtr += SrcPtrDiff;
+			}
 		}
 	}
-	else if (bmhdr->biPlanes==1 && bmhdr->biBitCount==16)
-	{
-		UE_LOG(LogImageWrapper, Error, TEXT("BMP 16 bit format no longer supported. Use terrain tools for importing/exporting heightmaps."));
-		return false;
-	}
 	else
 	{
-		UE_LOG(LogImageWrapper, Error, TEXT("BMP uses an unsupported format (planes = %i, bitcount = %i)"), bmhdr->biPlanes, bmhdr->biBitCount);
+		UE_LOG(LogImageWrapper, Error, TEXT("BMP uses an unsupported format (planes=%i, bitcount=%i, compression=%i)"), 
+			bmhdr->biPlanes, bmhdr->biBitCount, bmhdr->biCompression);
 		return false;
 	}
 
@@ -420,13 +501,13 @@ bool FBmpImageWrapper::LoadBMPInfoHeader(int64 HeaderOffset)
 
 	const FBitmapInfoHeader* bmhdr = (FBitmapInfoHeader *)(CompressedData.GetData() + HeaderOffset);
 
-	if (bmhdr->biCompression != BCBI_RGB && bmhdr->biCompression != BCBI_BITFIELDS)
+	if (bmhdr->biCompression != BCBI_RGB && bmhdr->biCompression != BCBI_BITFIELDS && bmhdr->biCompression != BCBI_ALPHABITFIELDS)
 	{
 		UE_LOG(LogImageWrapper, Error, TEXT("RLE compression of BMP images not supported"));
 		return false;
 	}
 
-	if (bmhdr->biPlanes==1 && (bmhdr->biBitCount==8 || bmhdr->biBitCount==24 || bmhdr->biBitCount==32))
+	if (bmhdr->biPlanes==1 && (bmhdr->biBitCount==8 || bmhdr->biBitCount==16 || bmhdr->biBitCount==24 || bmhdr->biBitCount==32))
 	{
 		// Set texture properties.
 		Width = bmhdr->biWidth;
@@ -446,14 +527,8 @@ bool FBmpImageWrapper::LoadBMPInfoHeader(int64 HeaderOffset)
 
 		return true;
 	}
-	
-	if (bmhdr->biPlanes == 1 && bmhdr->biBitCount == 16)
-	{
-		UE_LOG(LogImageWrapper, Error, TEXT("BMP 16 bit format no longer supported. Use terrain tools for importing/exporting heightmaps."));
-		return false;
-	}
 	else
-	{
+	{	
 		UE_LOG(LogImageWrapper, Error, TEXT("BMP uses an unsupported format (planes = %i, bitcount = %i)"), bmhdr->biPlanes, bmhdr->biBitCount);
 		return false;
 	}
