Handle missing error handling case for FPngImageWrapper

diff --git a/Engine/Source/Runtime/ImageWrapper/Private/Formats/PngImageWrapper.cpp b/Engine/Source/Runtime/ImageWrapper/Private/Formats/PngImageWrapper.cpp
--- a/Engine/Source/Runtime/ImageWrapper/Private/Formats/PngImageWrapper.cpp
+++ b/Engine/Source/Runtime/ImageWrapper/Private/Formats/PngImageWrapper.cpp
@@ -513,7 +513,6 @@ bool FPngImageWrapper::IsPNG() const
 	return false;
 }
 
-
 bool FPngImageWrapper::LoadPNGHeader()
 {
 	check(CompressedData.Num());
@@ -531,6 +530,22 @@ bool FPngImageWrapper::LoadPNGHeader()
 		check(info_ptr);
 
 		PNGReadGuard PNGGuard(&png_ptr, &info_ptr);
+
+		// Store the current stack pointer in the jump buffer. setjmp will return non-zero in the case of a read error.
+#if PLATFORM_ANDROID
+		//Preserve old single thread code on some platform in relation to a type incompatibility at compile time.
+		if (setjmp(SetjmpBuffer) != 0)
+#else
+		//Use libPNG jump buffer solution to allow concurrent compression\decompression on concurrent threads.
+		if (setjmp(png_jmpbuf(png_ptr)) != 0)
+#endif
+		{
+			UE_LOG(LogImageWrapper, Error, TEXT("PNG Header Error"));
+			return false;
+		}
+
+		// ---------------------------------------------------------------------------------------------------------
+		// Anything allocated on the stack after this point will not be destructed correctly in the case of an error
 		{
 			png_set_read_fn(png_ptr, this, FPngImageWrapper::user_read_compressed);
 
