ImageWrapper: Fixed crash when reading exrs with a non zero data window.

diff --git a/Engine/Source/Runtime/ImageWrapper/Private/Formats/ExrImageWrapper.cpp b/Engine/Source/Runtime/ImageWrapper/Private/Formats/ExrImageWrapper.cpp
--- a/Engine/Source/Runtime/ImageWrapper/Private/Formats/ExrImageWrapper.cpp
+++ b/Engine/Source/Runtime/ImageWrapper/Private/Formats/ExrImageWrapper.cpp
@@ -495,17 +495,19 @@ void FExrImageWrapper::Uncompress(const ERGBFormat InFormat, const int32 InBitDe
 	try
 	{
 		Imf::FrameBuffer ImfFrameBuffer;
+		FMemFileIn MemFile(CompressedData.GetData(), CompressedData.Num());
+		Imf::InputFile ImfFile(MemFile);
+		Imf::Header ImfHeader = ImfFile.header();
+		Imath::Box2i ImfDataWindow = ImfHeader.dataWindow();
 		for (int32 c = 0; c < ChannelCount; ++c)
 		{
 			ChannelData[c].SetNumUninitialized((int64)BytesPerChannelPixel * Width * Height);
 			// Use 1.0 as a default value for the alpha channel, in case if it is not present in the EXR, use 0.0 for all other channels.
 			double DefaultValue = !strcmp(ChannelNames[c], "A") ? 1.0 : 0.0;
-			ImfFrameBuffer.insert(ChannelNames[c], Imf::Slice(ImfPixelType, (char*)ChannelData[c].GetData(), BytesPerChannelPixel, (size_t)BytesPerChannelPixel * Width, 1, 1, DefaultValue));
+			char* ChannelBase = (char*)(ChannelData[c].GetData() -
+				(ImfDataWindow.min.x + ImfDataWindow.min.y * Width) * ((int64)BytesPerChannelPixel));
+			ImfFrameBuffer.insert(ChannelNames[c], Imf::Slice(ImfPixelType, ChannelBase, BytesPerChannelPixel, (size_t)BytesPerChannelPixel * Width, 1, 1, DefaultValue));
 		}
-		FMemFileIn MemFile(CompressedData.GetData(), CompressedData.Num());
-		Imf::InputFile ImfFile(MemFile);
-		Imf::Header ImfHeader = ImfFile.header();
-		Imath::Box2i ImfDataWindow = ImfHeader.dataWindow();
 		ImfFile.setFrameBuffer(ImfFrameBuffer);
 		ImfFile.readPixels(ImfDataWindow.min.y, ImfDataWindow.max.y);
 	}
