sanitize floats for inf/nan in EXR import

diff --git a/Engine/Source/Runtime/Core/Public/Math/Float16.h b/Engine/Source/Runtime/Core/Public/Math/Float16.h
--- a/Engine/Source/Runtime/Core/Public/Math/Float16.h
+++ b/Engine/Source/Runtime/Core/Public/Math/Float16.h
@@ -88,6 +88,9 @@ public:
 	/** Return float clamp in [0,MaxF16Float] , no negatives or infinites or nans returned **/
 	FFloat16 GetClampedNonNegativeAndFinite() const;
 	
+	/** Return float clamp in [-MaxF16Float,MaxF16Float] , no infinites or nans returned **/
+	FFloat16 GetClampedFinite() const;
+
 	/** Convert from Fp16 to Fp32. */
 	float GetFloat() const;
 
@@ -249,3 +252,34 @@ FORCEINLINE FFloat16 FFloat16::GetClampedNonNegativeAndFinite() const
 
 	return ReturnValue;
 }
+
+
+/** Return float clamp in [-MaxF16Float,MaxF16Float] , no infinites or nans returned **/
+FORCEINLINE FFloat16 FFloat16::GetClampedFinite() const
+{	
+	FFloat16 ReturnValue;
+
+	if ( (Encoded&0x7c00) == 0x7c00 )
+	{
+		// inf or nan
+		if ( Encoded == 0x7C00 ) //+inf
+		{
+			ReturnValue.Encoded = 0x7bff; // max finite
+		}
+		else if ( Encoded == 0xFC00 ) //-inf
+		{
+			ReturnValue.Encoded = 0xfbff; // max finite negative
+		}
+		else
+		{
+			// nan
+			ReturnValue.Encoded = 0;
+		}
+	}
+	else
+	{
+		ReturnValue.Encoded = Encoded;
+	}
+	
+	return ReturnValue;
+}
\ No newline at end of file
diff --git a/Engine/Source/Runtime/ImageWrapper/Private/Formats/ExrImageWrapper.cpp b/Engine/Source/Runtime/ImageWrapper/Private/Formats/ExrImageWrapper.cpp
--- a/Engine/Source/Runtime/ImageWrapper/Private/Formats/ExrImageWrapper.cpp
+++ b/Engine/Source/Runtime/ImageWrapper/Private/Formats/ExrImageWrapper.cpp
@@ -504,6 +504,7 @@ void FExrImageWrapper::Uncompress(const ERGBFormat InFormat, const int32 InBitDe
 			ChannelData[c].SetNumUninitialized((int64)BytesPerChannelPixel * Width * Height);
 			// Use 1.0 as a default value for the alpha channel, in case if it is not present in the EXR, use 0.0 for all other channels.
 			double DefaultValue = !strcmp(ChannelNames[c], "A") ? 1.0 : 0.0;
+			// @todo Oodle: doing a subtract on this pointer for the data window looks dangerous
 			char* ChannelBase = (char*)(ChannelData[c].GetData() -
 				(ImfDataWindow.min.x + ImfDataWindow.min.y * Width) * ((int64)BytesPerChannelPixel));
 			ImfFrameBuffer.insert(ChannelNames[c], Imf::Slice(ImfPixelType, ChannelBase, BytesPerChannelPixel, (size_t)BytesPerChannelPixel * Width, 1, 1, DefaultValue));
@@ -520,14 +521,76 @@ void FExrImageWrapper::Uncompress(const ERGBFormat InFormat, const int32 InBitDe
 	}
 
 	// EXR channels are compressed non-interleaved.
-	RawData.SetNumUninitialized((int64)BytesPerChannelPixel * ChannelCount * Width * Height);
-	for (int64 OffsetNonInterleaved = 0, OffsetInterleaved = 0; OffsetInterleaved < RawData.Num(); OffsetNonInterleaved += BytesPerChannelPixel)
+	int64 BytesPerChannel = (int64)BytesPerChannelPixel * Width * Height;
+	RawData.SetNumUninitialized(BytesPerChannel * ChannelCount);
+
+	const uint8 * ChannelDataPointers[4];
+	check(ChannelCount == 1 || ChannelCount == 4);
+	for (int32 c = 0; c < ChannelCount; ++c)
 	{
-		for (int32 c = 0; c < ChannelCount; ++c)
+		ChannelDataPointers[c] = ChannelData[c].GetData();
+	}
+
+	if ( InBitDepth == 16 )
+	{
+		FFloat16 * Out = (FFloat16 *)&RawData[0];
+
+		for (int64 OffsetNonInterleaved = 0; OffsetNonInterleaved < BytesPerChannel; OffsetNonInterleaved += 2)
 		{
-			for (int32 b = 0; b < BytesPerChannelPixel; ++b, ++OffsetInterleaved)
+			for (int32 c = 0; c < ChannelCount; ++c)
+			{
+				const FFloat16 * In = (const FFloat16 *)&ChannelDataPointers[c][OffsetNonInterleaved];
+				
+				*Out = In->GetClampedFinite();
+
+				//check( ! FMath::IsNaN( Out->GetFloat() ) );
+				//check( Out->GetFloat() == In->GetFloat() || ! isfinite( In->GetFloat() ) );
+
+				Out++;
+			}
+		}
+	}
+	else
+	{
+		check( InBitDepth == 32 );
+		
+		float * Out = (float *)&RawData[0];
+		
+		for (int64 OffsetNonInterleaved = 0; OffsetNonInterleaved < BytesPerChannel; OffsetNonInterleaved += 4)
+		{
+			for (int32 c = 0; c < ChannelCount; ++c)
 			{
-				RawData[OffsetInterleaved] = ChannelData[c][OffsetNonInterleaved + b];
+				const float * In = (const float *)&ChannelDataPointers[c][OffsetNonInterleaved];
+				
+				float f = *In;
+
+				// sanitize inf and nan :
+				if ( f >= -FLT_MAX && f <= FLT_MAX )
+				{
+					// finite, leave it
+					// nans will fail all compares so not go in here
+				}
+				else if ( f > FLT_MAX )
+				{
+					// +inf
+					f = FLT_MAX;
+				}
+				else if ( f < -FLT_MAX )
+				{
+					// -inf
+					f = -FLT_MAX;
+				}
+				else
+				{
+					// nan
+					f = 0.f;
+				}
+
+				*Out = f;
+				
+				//check( ! FMath::IsNaN( *Out ) );
+
+				Out++;
 			}
 		}
 	}
