ImageWrapper : remove pointless SetRaw implementations
just fall through to FImageWrapperBase::SetRaw

diff --git a/Engine/Source/Runtime/ImageWrapper/Private/Formats/ExrImageWrapper.cpp b/Engine/Source/Runtime/ImageWrapper/Private/Formats/ExrImageWrapper.cpp
--- a/Engine/Source/Runtime/ImageWrapper/Private/Formats/ExrImageWrapper.cpp
+++ b/Engine/Source/Runtime/ImageWrapper/Private/Formats/ExrImageWrapper.cpp
@@ -222,48 +222,6 @@ ERawImageFormat::Type FExrImageWrapper::GetSupportedRawFormat(const ERawImageFor
 	}
 }
 
-bool FExrImageWrapper::SetRaw(const void* InRawData, int64 InRawSize, const int32 InWidth, const int32 InHeight, const ERGBFormat InFormat, const int32 InBitDepth, const int32 InBytesPerRow)
-{
-	check(InRawData);
-	check(InRawSize > 0);
-	check(InWidth > 0);
-	check(InHeight > 0);
-	check(InBytesPerRow >= 0);
-
-	// FExrImageWrapper used to take RGBA 8-bit input
-	//	 and write it linearly
-	// the new image path now requires you to convert to float before coming in here
-	// so U8 will be converted to float *with* gamma correction
-
-	switch (InBitDepth)
-	{
-	case 8:
-		if (InFormat != ERGBFormat::RGBA && InFormat != ERGBFormat::BGRA && InFormat != ERGBFormat::Gray)
-		{
-			return false;
-		}
-		break;
-
-	case 16:
-	case 32:
-		if (InFormat == ERGBFormat::RGBA || InFormat == ERGBFormat::Gray)
-		{
-			// Before ERGBFormat::RGBAF and ERGBFormat::GrayF were introduced, ERGBFormat::RGBA and ERGBFormat::Gray were used to describe float pixel formats.
-			// ERGBFormat::RGBA and ERGBFormat::Gray should now only be used for integer channels.
-			// Note that EXR uint32 compression is currently not supported.
-			const TCHAR* FormatName = (InFormat == ERGBFormat::RGBA) ? TEXT("RGBA") : TEXT("Gray");
-			UE_LOG(LogImageWrapper, Warning, TEXT("Usage of 16-bit and 32-bit ERGBFormat::%s raw format for compressing EXR images is deprecated, if you are compressing float channels please specify ERGBFormat::%sF instead."), *FormatName, *FormatName);
-		}
-		if (InFormat != ERGBFormat::RGBAF && InFormat != ERGBFormat::GrayF)
-		{
-			return false;
-		}
-		break;
-	}
-
-	return FImageWrapperBase::SetRaw(InRawData, InRawSize, InWidth, InHeight, InFormat, InBitDepth, InBytesPerRow);
-}
-
 bool FExrImageWrapper::SetCompressed(const void* InCompressedData, int64 InCompressedSize)
 {
 	check(InCompressedData);
@@ -606,26 +564,6 @@ FExrImageWrapper::FExrImageWrapper()
 {
 }
 
-bool FExrImageWrapper::SetRaw(const void* InRawData, int64 InRawSize, const int32 InWidth, const int32 InHeight, const ERGBFormat InFormat, const int32 InBitDepth, const int32 InBytesPerRow)
-{
-	check(InRawData);
-	check(InRawSize > 0);
-	check(InWidth > 0);
-	check(InHeight > 0);
-	check(InBytesPerRow >= 0);
-	if (InBitDepth != 32)
-	{
-		return false;
-	}
-
-	if (InFormat != ERGBFormat::RGBAF)
-	{
-		return false;
-	}
-
-	return FImageWrapperBase::SetRaw(InRawData, InRawSize, InWidth, InHeight, InFormat, InBitDepth, InBytesPerRow);
-}
-
 bool FExrImageWrapper::SetCompressed(const void* InCompressedData, int64 InCompressedSize)
 {
 	return false;
@@ -643,6 +581,7 @@ void FExrImageWrapper::Compress(int32 Quality)
 
 	const double StartTime = FPlatformTime::Seconds();
 
+	// conditions enforced by CanSetRawFormat :
 	check(BitDepth == 32);
 	check(Format == ERGBFormat::RGBAF);
 
diff --git a/Engine/Source/Runtime/ImageWrapper/Private/Formats/ExrImageWrapper.h b/Engine/Source/Runtime/ImageWrapper/Private/Formats/ExrImageWrapper.h
--- a/Engine/Source/Runtime/ImageWrapper/Private/Formats/ExrImageWrapper.h
+++ b/Engine/Source/Runtime/ImageWrapper/Private/Formats/ExrImageWrapper.h
@@ -41,7 +41,6 @@ public:
 
 	//~ FImageWrapper interface
 
-	virtual bool SetRaw(const void* InRawData, int64 InRawSize, const int32 InWidth, const int32 InHeight, const ERGBFormat InFormat, const int32 InBitDepth, const int32 InBytesPerRow = 0) override;
 	virtual bool SetCompressed(const void* InCompressedData, int64 InCompressedSize) override;
 	virtual void Compress(int32 Quality) override;
 	virtual void Uncompress(const ERGBFormat InFormat, int32 InBitDepth) override;
diff --git a/Engine/Source/Runtime/ImageWrapper/Private/Formats/IcnsImageWrapper.cpp b/Engine/Source/Runtime/ImageWrapper/Private/Formats/IcnsImageWrapper.cpp
--- a/Engine/Source/Runtime/ImageWrapper/Private/Formats/IcnsImageWrapper.cpp
+++ b/Engine/Source/Runtime/ImageWrapper/Private/Formats/IcnsImageWrapper.cpp
@@ -72,6 +72,7 @@ bool FIcnsImageWrapper::SetRaw(const void* InRawData, int64 InRawSize, const int
 	check(InBytesPerRow == 0);
 
 #if PLATFORM_MAC
+	// CanSetRawFormat returns false, this will always fail, remove?
 	return FImageWrapperBase::SetRaw(InRawData, InRawSize, InWidth, InHeight, InFormat, InBitDepth, InBytesPerRow);
 #else
 	return false;
diff --git a/Engine/Source/Runtime/ImageWrapper/Private/Formats/JpegImageWrapper.cpp b/Engine/Source/Runtime/ImageWrapper/Private/Formats/JpegImageWrapper.cpp
--- a/Engine/Source/Runtime/ImageWrapper/Private/Formats/JpegImageWrapper.cpp
+++ b/Engine/Source/Runtime/ImageWrapper/Private/Formats/JpegImageWrapper.cpp
@@ -162,16 +162,6 @@ ERawImageFormat::Type FJpegImageWrapper::GetSupportedRawFormat(const ERawImageFo
 	};
 }
 
-bool FJpegImageWrapper::SetRaw(const void* InRawData, int64 InRawSize, const int32 InWidth, const int32 InHeight, const ERGBFormat InFormat, const int32 InBitDepth, const int32 InBytesPerRow)
-{
-	check((InFormat == ERGBFormat::RGBA || InFormat == ERGBFormat::BGRA || InFormat == ERGBFormat::Gray) && InBitDepth == 8);
-
-	bool bResult = FImageWrapperBase::SetRaw(InRawData, InRawSize, InWidth, InHeight, InFormat, InBitDepth, InBytesPerRow);
-
-	return bResult;
-}
-
-
 void FJpegImageWrapper::Compress(int32 Quality)
 {
 	if (Quality == 0)
@@ -378,6 +368,7 @@ void FJpegImageWrapper::UncompressTurbo(const ERGBFormat InFormat, int32 InBitDe
 	RawData.AddUninitialized(Width * Height * Channels);
 	const int PixelFormat = ConvertTJpegPixelFormat(InFormat);
 	const int Flags = TJFLAG_NOREALLOC | TJFLAG_FASTDCT;
+	// @todo Oodle : evaluate whether TJFLAG_FASTDCT quality loss is worth the speed gained
 
 	if (tjDecompress2(Decompressor, CompressedData.GetData(), CompressedData.Num(), RawData.GetData(), Width, 0, Height, PixelFormat, Flags) != 0)
 	{
diff --git a/Engine/Source/Runtime/ImageWrapper/Private/Formats/JpegImageWrapper.h b/Engine/Source/Runtime/ImageWrapper/Private/Formats/JpegImageWrapper.h
--- a/Engine/Source/Runtime/ImageWrapper/Private/Formats/JpegImageWrapper.h
+++ b/Engine/Source/Runtime/ImageWrapper/Private/Formats/JpegImageWrapper.h
@@ -33,7 +33,6 @@ public:
 	//~ FImageWrapperBase interface
 
 	virtual bool SetCompressed(const void* InCompressedData, int64 InCompressedSize) override;
-	virtual bool SetRaw(const void* InRawData, int64 InRawSize, const int32 InWidth, const int32 InHeight, const ERGBFormat InFormat, const int32 InBitDepth, const int32 InBytesPerRow = 0) override;
 	virtual void Uncompress(const ERGBFormat InFormat, int32 InBitDepth) override;
 	virtual void Compress(int32 Quality) override;
 	
