ImageWrapper : do check conditions in release builds for security concern

diff --git a/Engine/Source/Runtime/ImageWrapper/Private/ImageWrapperBase.cpp b/Engine/Source/Runtime/ImageWrapper/Private/ImageWrapperBase.cpp
--- a/Engine/Source/Runtime/ImageWrapper/Private/ImageWrapperBase.cpp
+++ b/Engine/Source/Runtime/ImageWrapper/Private/ImageWrapperBase.cpp
@@ -95,11 +95,13 @@ bool FImageWrapperBase::SetCompressed(const void* InCompressedData, int64 InComp
 
 bool FImageWrapperBase::SetRaw(const void* InRawData, int64 InRawSize, const int32 InWidth, const int32 InHeight, const ERGBFormat InFormat, const int32 InBitDepth, const int32 InBytesPerRow)
 {
-	check(InRawData != NULL);
-	check(InRawSize > 0);
-	check(InWidth > 0);
-	check(InHeight > 0);
-	check(InBytesPerRow >= 0);
+	#define check_and_return_false(exp)		do { check(exp); if ( ! (exp) ) return false; } while(0)
+
+	check_and_return_false(InRawData != NULL);
+	check_and_return_false(InRawSize > 0);
+	check_and_return_false(InWidth > 0);
+	check_and_return_false(InHeight > 0);
+	check_and_return_false(InBytesPerRow >= 0);
 
 	Reset();
 	CompressedData.Empty();		// Invalidates the compressed data too
@@ -117,32 +119,35 @@ bool FImageWrapperBase::SetRaw(const void* InRawData, int64 InRawSize, const int
 
 	int BytesPerRow = GetBytesPerRow();
 	
-
-	int64 RawSize = BytesPerRow * Height;
+	int64 RawSize = (int64) BytesPerRow * Height;
 	RawData.Empty(RawSize);
 	RawData.AddUninitialized(RawSize);
 
 	// copy the incoming data directly
-	if ((InBytesPerRow == 0 || BytesPerRow == InBytesPerRow) && (RawSize == InRawSize))
+	if ( InBytesPerRow == 0 || BytesPerRow == InBytesPerRow )
 	{
 		// this is usually an unnecessary allocation and copy
 		// we should compress directly from the source buffer
 
-		FMemory::Memcpy(RawData.GetData(), InRawData, InRawSize);
+		check_and_return_false( InRawSize >= RawSize );
 
+		FMemory::Memcpy(RawData.GetData(), InRawData, RawSize);
 	}
 	else
 	{
 		// The supported image formats (PNG, BMP, etc) don't support strided data (although turbo jpeg does).
 		// Therefore we de-stride the data here so we can uniformly support strided input data with all the supported
 		// image formats.
-		check(RawSize < Height* InBytesPerRow);
+		check_and_return_false(InBytesPerRow > BytesPerRow); // equality handled in above branch
+		check_and_return_false(InRawSize >= (int64) (Height-1)*InBytesPerRow + BytesPerRow);
 
 		for (int32 y = 0; y < Height; y++)
 		{
-			FMemory::Memcpy(RawData.GetData() + (y * BytesPerRow), (uint8*)InRawData + (y * InBytesPerRow), BytesPerRow);
+			FMemory::Memcpy(RawData.GetData() + (y * BytesPerRow), (uint8*)InRawData + (int64) y * InBytesPerRow, BytesPerRow);
 		}
 	}
+	
+	#undef check_and_return_false
 
 	return true;
 }
