ExrImageWrapper : fix importing exr's with channel names other than RGBA
should use the channel names found in the file

diff --git a/Engine/Source/Runtime/ImageWrapper/Private/Formats/ExrImageWrapper.cpp b/Engine/Source/Runtime/ImageWrapper/Private/Formats/ExrImageWrapper.cpp
--- a/Engine/Source/Runtime/ImageWrapper/Private/Formats/ExrImageWrapper.cpp
+++ b/Engine/Source/Runtime/ImageWrapper/Private/Formats/ExrImageWrapper.cpp
@@ -8,14 +8,6 @@
 #include "HAL/PlatformTime.h"
 #include "Math/Float16.h"
 
-/**
-
-@todo Oodle : this is a flawed EXR importer
-it crashes on many images in the OpenEXR test set
-Blobbies.exr and spirals.exr among others
-
-**/
-
 #if WITH_UNREALEXR
 
 FExrImageWrapper::FExrImageWrapper()
@@ -157,12 +149,14 @@ private:
 	int64 Pos;
 };
 
+// these are the channel names we write
+//	we will read whatever channel names are in the file
 const char* cChannelNamesRGBA[] = { "R", "G", "B", "A" };
 const char* cChannelNamesBGRA[] = { "B", "G", "R", "A" };
 //const char* cChannelNamesGray[] = { "G" }; // is that green or gray ?
 const char* cChannelNamesGray[] = { "Y" }; // pretty sure "Y" is more standard for gray
 
-int32 GetChannelNames(ERGBFormat InRGBFormat, const char* const*& OutChannelNames)
+static int32 GetChannelNames(ERGBFormat InRGBFormat, const char* const*& OutChannelNames)
 {
 	int32 ChannelCount;
 
@@ -226,6 +220,9 @@ bool FExrImageWrapper::SetCompressed(const void* InCompressedData, int64 InCompr
 {
 	check(InCompressedData);
 	check(InCompressedSize > 0);
+	
+	// reset variables :
+	FileChannelNames.Empty();
 
 	// Check the magic value in advance to avoid spamming the log with EXR parsing errors.
 	if (InCompressedSize < sizeof(uint32) || *(uint32*)InCompressedData != Imf::MAGIC)
@@ -250,7 +247,6 @@ bool FExrImageWrapper::SetCompressed(const void* InCompressedData, int64 InCompr
 		Height = ImfDataWindow.max.y - ImfDataWindow.min.y + 1;
 
 		bool bHasOnlyHALFChannels = true;
-		bool bHasAlphaChannel = false;
 		int32 ChannelCount = 0;
 
 		for (Imf::ChannelList::Iterator Iter = ImfChannels.begin(); Iter != ImfChannels.end(); ++Iter)
@@ -265,14 +261,20 @@ bool FExrImageWrapper::SetCompressed(const void* InCompressedData, int64 InCompr
 			//  (we incorrectly used "G" before, we now use "Y", TinyEXR uses "A" for 1-channel EXR)
 			//bMatchesGrayOrder = bMatchesGrayOrder && ChannelCount < UE_ARRAY_COUNT(cChannelNamesGray) && !strcmp(Iter.name(), cChannelNamesGray[ChannelCount]);
 
-			if ( strcmp(Iter.name(),"A") == 0 )
-			{
-				bHasAlphaChannel = true;
-			}
+			FileChannelNames.Add( MakeUniqueCString(Iter.name()) );
+		}
+
+		if ( ChannelCount == 0 )
+		{
+			SetError( TEXT("EXR has no channels") );
+			return false;
 		}
 
 		BitDepth = (ChannelCount && bHasOnlyHALFChannels) ? 16 : 32;
 
+		char Channel0LastChar = FileChannelNames[0][ strlen(FileChannelNames[0].Get())-1 ];
+		bool bChannel0IsAlpha = toupper(Channel0LastChar) == 'A'; 
+
 		// EXR uint32 channels are currently not supported, therefore input channels are always treated as float channels.
 		// Channel combinations which don't match the ERGBFormat::GrayF pattern are qualified as ERGBFormat::RGBAF.
 		// Note that channels inside the EXR file are indexed by name, therefore can be decoded in any RGB order.
@@ -281,7 +283,7 @@ bool FExrImageWrapper::SetCompressed(const void* InCompressedData, int64 InCompr
 		//  this cannot be loaded as a 1-channel image (you would just get all zeros)
 		// it must be loaded as RGBA here
 		// @todo Oodle : load that as RGBA then move A to R and convert back to 1 channel ?
-		if (ChannelCount == 1 && !bHasAlphaChannel)
+		if (ChannelCount == 1 && !bChannel0IsAlpha)
 		{
 			Format = ERGBFormat::GrayF;
 		}
@@ -321,6 +323,7 @@ void FExrImageWrapper::Compress(int32 Quality)
 	{
 		// uint8 channels are linearly converted into FFloat16 channels.
 		// note: NO GAMMA CORRECTION
+		//	(this is no longer used in the modern FImage based APIs, conversion to float should be done before calling this)
 		ConvertedRawData.SetNumUninitialized(sizeof(FFloat16) * RawData.Num());
 		FFloat16* Output = reinterpret_cast<FFloat16*>(ConvertedRawData.GetData());
 		for (int64 i = 0; i < RawData.Num(); ++i)
@@ -439,8 +442,84 @@ void FExrImageWrapper::Uncompress(const ERGBFormat InFormat, const int32 InBitDe
 		return;
 	}
 
-	const char* const* ChannelNames;
-	int32 ChannelCount = GetChannelNames(InFormat, ChannelNames);
+	// use FileChannelNames[] from SetCompressed
+	//	so that we ask for channels that are actually in the file
+	const char* ChannelNames[4] = { };
+	int32 ChannelCount;
+	if ( InFormat == ERGBFormat::GrayF )
+	{
+		check( FileChannelNames.Num() == 1 );
+		ChannelNames[0] = FileChannelNames[0].Get();
+		ChannelCount = 1;
+	}
+	else
+	{
+		const char * ChannelNamesDesiredChars = "RGBA";
+		ChannelCount = 4;
+
+		// try to find channel names in the file which correspond to RGBA
+		//	by looking at the last char of the file channel name
+		//	and trying to match it to ChannelNamesDesiredChars
+		//	(so that eg. "rgb.R" maps to channel 0)
+		TArray<const char *> FileChannelNamesUnused;
+		for(int FileChannelNamesI=0;FileChannelNamesI<FileChannelNames.Num();FileChannelNamesI++)
+		{
+			const char * Str = FileChannelNames[FileChannelNamesI].Get();
+			check( Str && *Str );
+			char LastChar = toupper( Str[ strlen(Str)-1 ] );
+			
+			for(int ChannelIndex=0;ChannelIndex<4;ChannelIndex++)
+			{
+				char DesiredChar = ChannelNamesDesiredChars[ChannelIndex];
+				if ( LastChar == DesiredChar && ChannelNames[ChannelIndex] == nullptr )
+				{
+					ChannelNames[ChannelIndex] = Str;
+					Str = nullptr;
+					break;
+				}
+			}
+
+			if ( Str )
+			{
+				FileChannelNamesUnused.Add( Str );
+			}
+		}
+
+		// if there are channel names in the file which were not mapped
+		//	and we have not found 4 channels
+		//	then just stuff the file channel names into channels in arbitrary order
+		//	eg. if the EXR has channel names like "X","Y","Z" lets go ahead and load them into RGB
+		for(int ChannelIndex=0;ChannelIndex<4;ChannelIndex++)
+		{
+			if ( ChannelNames[ChannelIndex] )
+			{
+				continue;
+			}
+			
+			if ( ! FileChannelNamesUnused.IsEmpty() )
+			{
+				ChannelNames[ChannelIndex] = FileChannelNamesUnused.Pop(false);
+			}
+			else
+			{
+				// stuff something that is not nullptr
+				//	and won't be found in the file
+				//	channel will be filled with default value
+				const char * DefaultChannelNames[4] =
+				{
+					"default.R","default.G","default.B","default.A"
+				};
+				ChannelNames[ChannelIndex] = DefaultChannelNames[ChannelIndex];
+			}
+		}
+				
+		UE_LOG(LogImageWrapper, Verbose, TEXT("Reading EXR with Channel Names: %s %s %s %s"), 
+			ANSI_TO_TCHAR(ChannelNames[0]),
+			ANSI_TO_TCHAR(ChannelNames[1]),
+			ANSI_TO_TCHAR(ChannelNames[2]),
+			ANSI_TO_TCHAR(ChannelNames[3]));
+	}
+
 	check(ChannelCount == 1 || ChannelCount == 4);
 
 	TArray<TArray64<uint8>> ChannelData;
@@ -460,11 +539,20 @@ void FExrImageWrapper::Uncompress(const ERGBFormat InFormat, const int32 InBitDe
 		for (int32 c = 0; c < ChannelCount; ++c)
 		{
 			ChannelData[c].SetNumUninitialized((int64)BytesPerChannelPixel * Width * Height);
+
+			// if you ask for a channel name that is not in the file data, you will get back DefaultValue
 			// Use 1.0 as a default value for the alpha channel, in case if it is not present in the EXR, use 0.0 for all other channels.
-			double DefaultValue = !strcmp(ChannelNames[c], "A") ? 1.0 : 0.0;
+			bool bIsAlphaChannel = false;
+			if ( c == 3 )
+			{
+				bIsAlphaChannel = toupper( ChannelNames[c][ strlen(ChannelNames[c]) -1 ] ) == 'A'; 
+			}
+			double DefaultValue = bIsAlphaChannel ? 1.0 : 0.0;
+
 			// @todo Oodle: doing a subtract on this pointer for the data window looks dangerous
 			char* ChannelBase = (char*)(ChannelData[c].GetData() -
 				(ImfDataWindow.min.x + ImfDataWindow.min.y * Width) * ((int64)BytesPerChannelPixel));
+
 			ImfFrameBuffer.insert(ChannelNames[c], Imf::Slice(ImfPixelType, ChannelBase, BytesPerChannelPixel, (size_t)BytesPerChannelPixel * Width, 1, 1, DefaultValue));
 		}
 		ImfFile.setFrameBuffer(ImfFrameBuffer);
diff --git a/Engine/Source/Runtime/ImageWrapper/Private/Formats/ExrImageWrapper.h b/Engine/Source/Runtime/ImageWrapper/Private/Formats/ExrImageWrapper.h
--- a/Engine/Source/Runtime/ImageWrapper/Private/Formats/ExrImageWrapper.h
+++ b/Engine/Source/Runtime/ImageWrapper/Private/Formats/ExrImageWrapper.h
@@ -47,6 +47,16 @@ public:
 	
 	virtual bool CanSetRawFormat(const ERGBFormat InFormat, const int32 InBitDepth) const override;
 	virtual ERawImageFormat::Type GetSupportedRawFormat(const ERawImageFormat::Type InFormat) const override;
+
+	typedef TUniquePtr<char[]> FUniqueCString;
+	FUniqueCString MakeUniqueCString(const char *str)
+	{
+		size_t num = strlen(str)+1;
+		FUniqueCString Ret = MakeUnique<char[]>( num );
+		memcpy(Ret.Get(),str,num);
+		return Ret;
+	}
+	TArray< FUniqueCString > FileChannelNames;
 };
 
 #endif // WITH_UNREALEXR
