Exr import : don't promote 1 channel exr to 4 channel in the TSF uasset
not needed any more

diff --git a/Engine/Source/Runtime/ImageWrapper/Private/Formats/ExrImageWrapper.cpp b/Engine/Source/Runtime/ImageWrapper/Private/Formats/ExrImageWrapper.cpp
--- a/Engine/Source/Runtime/ImageWrapper/Private/Formats/ExrImageWrapper.cpp
+++ b/Engine/Source/Runtime/ImageWrapper/Private/Formats/ExrImageWrapper.cpp
@@ -255,12 +255,6 @@ bool FExrImageWrapper::SetCompressed(const void* InCompressedData, int64 InCompr
 
 			bHasOnlyHALFChannels = bHasOnlyHALFChannels && Iter.channel().type == Imf::HALF;
 
-			// check for bMatchesGrayOrder disabled
-			// treat any 1-channel import as gray
-			// don't try to match the channel names, which are not standardized
-			//  (we incorrectly used "G" before, we now use "Y", TinyEXR uses "A" for 1-channel EXR)
-			//bMatchesGrayOrder = bMatchesGrayOrder && ChannelCount < UE_ARRAY_COUNT(cChannelNamesGray) && !strcmp(Iter.name(), cChannelNamesGray[ChannelCount]);
-
 			FileChannelNames.Add( MakeUniqueCString(Iter.name()) );
 		}
 
@@ -272,18 +266,10 @@ bool FExrImageWrapper::SetCompressed(const void* InCompressedData, int64 InCompr
 
 		BitDepth = (ChannelCount && bHasOnlyHALFChannels) ? 16 : 32;
 
-		char Channel0LastChar = FileChannelNames[0][ strlen(FileChannelNames[0].Get())-1 ];
-		bool bChannel0IsAlpha = toupper(Channel0LastChar) == 'A'; 
-
 		// EXR uint32 channels are currently not supported, therefore input channels are always treated as float channels.
-		// Channel combinations which don't match the ERGBFormat::GrayF pattern are qualified as ERGBFormat::RGBAF.
 		// Note that channels inside the EXR file are indexed by name, therefore can be decoded in any RGB order.
 
-		// NOTE: TinyEXR writes 1-channel EXR as an "A" named channel
-		//  this cannot be loaded as a 1-channel image (you would just get all zeros)
-		// it must be loaded as RGBA here
-		// @todo Oodle : load that as RGBA then move A to R and convert back to 1 channel ?
-		if (ChannelCount == 1 && !bChannel0IsAlpha)
+		if (ChannelCount == 1 )
 		{
 			Format = ERGBFormat::GrayF;
 		}
@@ -448,6 +434,8 @@ void FExrImageWrapper::Uncompress(const ERGBFormat InFormat, const int32 InBitDe
 	int32 ChannelCount;
 	if ( InFormat == ERGBFormat::GrayF )
 	{
+		// load the one channel into slot 0 regardless of name
+		// output format will be R16F or R32F , so it will look red
 		check( FileChannelNames.Num() == 1 );
 		ChannelNames[0] = FileChannelNames[0].Get();
 		ChannelCount = 1;
@@ -543,8 +531,11 @@ void FExrImageWrapper::Uncompress(const ERGBFormat InFormat, const int32 InBitDe
 			// if you ask for a channel name that is not in the file data, you will get back DefaultValue
 			// Use 1.0 as a default value for the alpha channel, in case if it is not present in the EXR, use 0.0 for all other channels.
 			bool bIsAlphaChannel = false;
+			// only treat a channel named "A" as alpha if it got in slot 3 of RGBA
+			//	for example TinyEXR writes out all 1-channel gray images with a channel name "A", we do not treat that as alpha
 			if ( c == 3 )
 			{
+				// should be true for DefaultChannelNames[3]
 				bIsAlphaChannel = toupper( ChannelNames[c][ strlen(ChannelNames[c]) -1 ] ) == 'A'; 
 			}
 			double DefaultValue = bIsAlphaChannel ? 1.0 : 0.0;
