Resave PNG files with incorrect sRGB profile
Log PNG filename during warnings and errors

diff --git a/Engine/Source/Runtime/ImageWrapper/Private/Formats/PngImageWrapper.cpp b/Engine/Source/Runtime/ImageWrapper/Private/Formats/PngImageWrapper.cpp
--- a/Engine/Source/Runtime/ImageWrapper/Private/Formats/PngImageWrapper.cpp
+++ b/Engine/Source/Runtime/ImageWrapper/Private/Formats/PngImageWrapper.cpp
@@ -618,7 +618,7 @@ void FPngImageWrapper::user_error_fn(png_structp png_ptr, png_const_charp error_
 		FString ErrorMsg = ANSI_TO_TCHAR(error_msg);
 		ctx->SetError(*ErrorMsg);
 
-		UE_LOG(LogImageWrapper, Error, TEXT("PNG Error: %s"), *ErrorMsg);
+		UE_LOG(LogImageWrapper, Error, TEXT("PNG Error(%s): %s"), ctx->DebugImageName != nullptr ? ctx->DebugImageName : TEXT(""), *ErrorMsg);
 
 	#if !PLATFORM_EXCEPTIONS_DISABLED
 		/** 
@@ -647,7 +647,9 @@ void FPngImageWrapper::user_error_fn(png_structp png_ptr, png_const_charp error_
 
 void FPngImageWrapper::user_warning_fn(png_structp png_ptr, png_const_charp warning_msg)
 {
-	UE_LOG(LogImageWrapper, Warning, TEXT("PNG Warning: %s"), ANSI_TO_TCHAR(warning_msg));
+	FPngImageWrapper* ctx = (FPngImageWrapper*)png_get_error_ptr(png_ptr);
+	const TCHAR* LocalDebugImageName = ctx != nullptr ? ctx->DebugImageName : nullptr;
+	UE_LOG(LogImageWrapper, Warning, TEXT("PNG Warning(%s) %s"), LocalDebugImageName != nullptr ? LocalDebugImageName : TEXT(""), ANSI_TO_TCHAR(warning_msg));
 }
 
 void* FPngImageWrapper::user_malloc(png_structp /*png_ptr*/, png_size_t size)
diff --git a/Engine/Source/Runtime/ImageWrapper/Private/ImageWrapperModule.cpp b/Engine/Source/Runtime/ImageWrapper/Private/ImageWrapperModule.cpp
--- a/Engine/Source/Runtime/ImageWrapper/Private/ImageWrapperModule.cpp
+++ b/Engine/Source/Runtime/ImageWrapper/Private/ImageWrapperModule.cpp
@@ -74,7 +74,7 @@ public:
 
 	//~ IImageWrapperModule interface
 
-	virtual TSharedPtr<IImageWrapper> CreateImageWrapper(const EImageFormat InFormat) override
+	virtual TSharedPtr<IImageWrapper> CreateImageWrapper(const EImageFormat InFormat, const TCHAR* InOptionalDebugImageName = nullptr) override
 	{
 		TRACE_CPUPROFILER_EVENT_SCOPE(ImageWrapper.Create);
 
@@ -137,6 +137,11 @@ public:
 			break;
 		}
 
+		if (InOptionalDebugImageName != nullptr)
+		{
+			ImageWrapper->SetDebugImageName(InOptionalDebugImageName);
+		}
+
 		return ImageWrapper;
 	}
 
diff --git a/Engine/Source/Runtime/ImageWrapper/Public/IImageWrapper.h b/Engine/Source/Runtime/ImageWrapper/Public/IImageWrapper.h
--- a/Engine/Source/Runtime/ImageWrapper/Public/IImageWrapper.h
+++ b/Engine/Source/Runtime/ImageWrapper/Public/IImageWrapper.h
@@ -109,6 +109,10 @@ enum class EImageCompressionQuality : uint8
  */
 class IImageWrapper
 {
+protected:
+	// debug image name string for any errors or warnings
+	const TCHAR* DebugImageName =  nullptr;
+
 public:
 
 
@@ -334,6 +338,13 @@ public:
 		return Ret;
 	}
 
+	/* Set the debug image name
+	 */
+	void SetDebugImageName(const TCHAR* InDebugImageName)
+	{
+		DebugImageName = InDebugImageName;
+	}
+
 public:
 
 	/** Virtual destructor. */
diff --git a/Engine/Source/Runtime/ImageWrapper/Public/IImageWrapperModule.h b/Engine/Source/Runtime/ImageWrapper/Public/IImageWrapperModule.h
--- a/Engine/Source/Runtime/ImageWrapper/Public/IImageWrapperModule.h
+++ b/Engine/Source/Runtime/ImageWrapper/Public/IImageWrapperModule.h
@@ -59,12 +59,13 @@ public:
 	 * Create an IImageWrapper helper of a specific type 
 	 *
 	 * @param InFormat - The type of image we want to deal with
-	 * @return The helper base class to manage the data
+	 * @param InOptionalDebugImageName - An optional string to be displayed with any errors or warnings
 	 * 
-	 * EImageFormat is a compressor / file format , not a pixel format
+	 * @return The helper base class to manage the data
+	 * EImageFormat is a compressor / file format, not a pixel format
 	 * Deprecated.  Prefer CompressImage/DecompressImage.
 	 */
-	virtual TSharedPtr<IImageWrapper> CreateImageWrapper(const EImageFormat InFormat) = 0;
+	virtual TSharedPtr<IImageWrapper> CreateImageWrapper(const EImageFormat InFormat, const TCHAR* InOptionalDebugImageName = nullptr) = 0;
 
 	/**
 	 * Detect image format by looking at the first few bytes of the compressed image data.
diff --git a/Engine/Source/Runtime/SlateRHIRenderer/Private/SlateRHIResourceManager.cpp b/Engine/Source/Runtime/SlateRHIRenderer/Private/SlateRHIResourceManager.cpp
--- a/Engine/Source/Runtime/SlateRHIRenderer/Private/SlateRHIResourceManager.cpp
+++ b/Engine/Source/Runtime/SlateRHIRenderer/Private/SlateRHIResourceManager.cpp
@@ -525,7 +525,7 @@ bool FSlateRHIResourceManager::LoadTexture( const FName& TextureName, const FStr
 		{
 			ImageFormat = EImageFormat::PNG;
 		}
-		TSharedPtr<IImageWrapper> ImageWrapper = ImageWrapperModule.CreateImageWrapper(ImageFormat);
+		TSharedPtr<IImageWrapper> ImageWrapper = ImageWrapperModule.CreateImageWrapper(ImageFormat, *ResourcePath);
 
 		if ( ImageWrapper.IsValid() && ImageWrapper->SetCompressed( RawFileData.GetData(), RawFileData.Num()) )
 		{
