nop comments only better description of EImageCompressionQuality

diff --git a/Engine/Source/Runtime/ImageWrapper/Private/Formats/ExrImageWrapper.cpp b/Engine/Source/Runtime/ImageWrapper/Private/Formats/ExrImageWrapper.cpp
--- a/Engine/Source/Runtime/ImageWrapper/Private/Formats/ExrImageWrapper.cpp
+++ b/Engine/Source/Runtime/ImageWrapper/Private/Formats/ExrImageWrapper.cpp
@@ -190,6 +190,8 @@ static int32 GetChannelNames(ERGBFormat InRGBFormat, const char* const*& OutChan
 
 bool FExrImageWrapper::CanSetRawFormat(const ERGBFormat InFormat, const int32 InBitDepth) const
 {
+	// only set float formats
+	//	integer formats should be converted before reaching ImageWrapper
 	return (InFormat == ERGBFormat::RGBAF || InFormat == ERGBFormat::GrayF) && (InBitDepth == 16 || InBitDepth == 32);
 }
 
diff --git a/Engine/Source/Runtime/ImageWrapper/Public/IImageWrapper.h b/Engine/Source/Runtime/ImageWrapper/Public/IImageWrapper.h
--- a/Engine/Source/Runtime/ImageWrapper/Public/IImageWrapper.h
+++ b/Engine/Source/Runtime/ImageWrapper/Public/IImageWrapper.h
@@ -84,15 +84,18 @@ enum class ERGBFormat : int8
  * JPEG interprets Quality as 1-100
  * JPEG default quality is 85 , Uncompressed means 100
  * 
- * Negative qualities set PNG zlib level
+ * for PNG:
+ * Negative qualities in [-1,-9] set PNG zlib level
  * PNG interprets "Uncompressed" as zlib level 0 (none)
+ * otherwise default zlib level 3 is used.
  * 
- * EXR respects the "Uncompressed" flag
+ * EXR respects the "Uncompressed" flag to turn off compression; otherwise ZIP_COMPRESSION is used.
  */
-enum class EImageCompressionQuality : uint8
+enum class EImageCompressionQuality : int8
 {
 	Default = 0,
 	Uncompressed = 1,
+	Max = 100,
 };
 
 
