Force "up" PNG filter when we are uncompressed for speed (and oddly, size too)

diff --git a/Engine/Source/Runtime/ImageWrapper/Private/Formats/PngImageWrapper.cpp b/Engine/Source/Runtime/ImageWrapper/Private/Formats/PngImageWrapper.cpp
--- a/Engine/Source/Runtime/ImageWrapper/Private/Formats/PngImageWrapper.cpp
+++ b/Engine/Source/Runtime/ImageWrapper/Private/Formats/PngImageWrapper.cpp
@@ -222,6 +222,15 @@ void FPngImageWrapper::Compress(int32 Quality)
 			png_set_IHDR(png_ptr, info_ptr, Width, Height, BitDepth, (Format == ERGBFormat::Gray) ? PNG_COLOR_TYPE_GRAY : PNG_COLOR_TYPE_RGBA, PNG_INTERLACE_NONE, PNG_COMPRESSION_TYPE_DEFAULT, PNG_FILTER_TYPE_DEFAULT);
 			png_set_write_fn(png_ptr, this, FPngImageWrapper::user_write_compressed, FPngImageWrapper::user_flush_data);
 
+			// If we're writing an uncompressed PNG, then we're expecting to be compressing
+			// externally and we can assume we're interested in speed. In this case, we force
+			// the fastest filter to avoid things like Paeth. This is ~2x speedup in decompressing
+			// bulk data encoded in this manner.
+			if (Quality == (int32)EImageCompressionQuality::Uncompressed)
+			{
+				png_set_filter(png_ptr, 0, PNG_FILTER_UP | PNG_FILTER_VALUE_UP);
+			}
+
 			const uint64 PixelChannels = (Format == ERGBFormat::Gray) ? 1 : 4;
 			const uint64 BytesPerPixel = (BitDepth * PixelChannels) / 8;
 			const uint64 BytesPerRow = BytesPerPixel * Width;
