Bounds checking for Tiff importing.

diff --git a/Engine/Source/Runtime/ImageWrapper/Private/Formats/TiffImageWrapper.cpp b/Engine/Source/Runtime/ImageWrapper/Private/Formats/TiffImageWrapper.cpp
--- a/Engine/Source/Runtime/ImageWrapper/Private/Formats/TiffImageWrapper.cpp
+++ b/Engine/Source/Runtime/ImageWrapper/Private/Formats/TiffImageWrapper.cpp
@@ -8,6 +8,7 @@
 #include "Async/TaskGraphInterfaces.h"
 #include "Containers/Queue.h"
 #include "CoreMinimal.h"
+#include "GuardedInt.h"
 #include "ImageWrapperPrivate.h"
 #include "Math/Float16.h"
 #include "Math/NumericLimits.h"
@@ -339,16 +340,24 @@ namespace UE::ImageWrapper::Private
 		static tmsize_t Read(thandle_t Handle, void* Buffer, tmsize_t Size)
 		{
 			FTiffImageWrapper* TiffImageWrapper = reinterpret_cast<FTiffImageWrapper*>(Handle);
-			int64 RemaningBytes = TiffImageWrapper->CompressedData.Num() - TiffImageWrapper->CurrentPosition;
-			if (RemaningBytes > 0)
+			if (TiffImageWrapper->CurrentPosition < 0 || TiffImageWrapper->CurrentPosition >= TiffImageWrapper->CompressedData.Num())
 			{
-				int64 NumBytesRead = FMath::Min<int64>(RemaningBytes, Size);
-				FMemory::Memcpy(Buffer, TiffImageWrapper->CompressedData.GetData() + TiffImageWrapper->CurrentPosition, NumBytesRead);
-				TiffImageWrapper->CurrentPosition += NumBytesRead;
+				return 0;
+			}
 
-				return NumBytesRead;
+			if (IntFitsIn<int64>(Size) == false)
+			{
+				Size = TNumericLimits<int64>::Max();
 			}
-			return 0;
+
+			// We now know this is [0, Num()].
+			int64 RemainingBytes = TiffImageWrapper->CompressedData.Num() - TiffImageWrapper->CurrentPosition;
+			int64 NumBytesRead = FMath::Min<int64>(RemainingBytes, Size);
+
+			FMemory::Memcpy(Buffer, TiffImageWrapper->CompressedData.GetData() + TiffImageWrapper->CurrentPosition, NumBytesRead);
+			TiffImageWrapper->CurrentPosition += NumBytesRead;
+
+			return NumBytesRead;
 		}
 
 		static tmsize_t Write(thandle_t Handle, void* Buffer, tmsize_t Size)
@@ -365,33 +374,29 @@ namespace UE::ImageWrapper::Private
 			const int OffsetFromCurrent = 1;
 			const int FromEnd = 2;
 
+			FGuardedInt64 TargetPosition;
 			switch (Whence)
 			{
 				case Set:
-						TiffImageWrapper->CurrentPosition = Offset;
-					break;
+					{
+						TargetPosition = FGuardedInt64(Offset);
+						break;
+					}
 				case OffsetFromCurrent:
 					{
-						const int64 NewPosition = TiffImageWrapper->CurrentPosition + Offset;
-						if (NewPosition >= 0)
-						{
-							TiffImageWrapper->CurrentPosition = NewPosition;
-						}
+						TargetPosition = FGuardedInt64(TiffImageWrapper->CurrentPosition) + Offset;
 						break;
 					}
 				case FromEnd:
 					{
-						const int64 NewPosition = TiffImageWrapper->CompressedData.Num() + Offset;
-						if (NewPosition >= 0)
-						{
-							TiffImageWrapper->CurrentPosition = NewPosition;
-						}
+						TargetPosition = FGuardedInt64(TiffImageWrapper->CompressedData.Num()) + Offset;
 						break;
 					}
 				default:
 					return -1;
 			}
 
+			TiffImageWrapper->CurrentPosition = FMath::Max(0, TargetPosition.Get(TiffImageWrapper->CurrentPosition));
 			return TiffImageWrapper->CurrentPosition;
 		}
 
@@ -446,7 +451,14 @@ namespace UE::ImageWrapper::Private
 			// Read using RGBA
 			if (Format == ERGBFormat::BGRA && BitDepth == 8)
 			{
-				const int64 BufferSize = int64(Width) * int64(Height) * sizeof(uint32);
+				FGuardedInt64 GuardedBufferSize = FGuardedInt64(Width) * Height * sizeof(uint32);
+				if (GuardedBufferSize.IsValid() == false)
+				{
+					SetError(TEXT("Tiff image is massive - likely corrupted file."));
+					return;
+				}
+
+				const int64 BufferSize = GuardedBufferSize.Get(0);
 				const int Flags = 0;
 	
 				RawData.Empty(BufferSize);
@@ -543,6 +555,13 @@ namespace UE::ImageWrapper::Private
 
 				Format = ERGBFormat::Invalid;
 				
+				if (Width <= 0 ||
+					Height <= 0)
+				{
+					SetError(TEXT("Invalid resolution in tiff: zero or negative."));
+					return false;
+				}
+
 				if ( SampleFormat != SAMPLEFORMAT_UINT
 					&& SampleFormat != SAMPLEFORMAT_IEEEFP
 					&& SampleFormat != 0 /* assume it's uint */)
