Security updates for PngImageWrapper

diff --git a/Engine/Source/Runtime/ImageWrapper/Private/Formats/PngImageWrapper.cpp b/Engine/Source/Runtime/ImageWrapper/Private/Formats/PngImageWrapper.cpp
--- a/Engine/Source/Runtime/ImageWrapper/Private/Formats/PngImageWrapper.cpp
+++ b/Engine/Source/Runtime/ImageWrapper/Private/Formats/PngImageWrapper.cpp
@@ -3,6 +3,7 @@
 #include "Formats/PngImageWrapper.h"
 #include "ImageWrapperPrivate.h"
 
+#include "Math/GuardedInt.h"
 #include "Misc/ScopeLock.h"
 
 
@@ -597,15 +598,23 @@ bool FPngImageWrapper::LoadPNGHeader()
 void FPngImageWrapper::user_read_compressed(png_structp png_ptr, png_bytep data, png_size_t length)
 {
 	FPngImageWrapper* ctx = (FPngImageWrapper*)png_get_io_ptr(png_ptr);
-	if (ctx->ReadOffset + (int64)length <= ctx->CompressedData.Num())
+	if (IntFitsIn<int64>(length) == false)
 	{
-		FMemory::Memcpy(data, &ctx->CompressedData[ctx->ReadOffset], length);
-		ctx->ReadOffset += length;
+		UE_LOG(LogImageWrapper, Warning, TEXT("Bad PNG read length: %llu"), length);
+		ctx->SetError(TEXT("Invalid length in read_compressed")); // this doesn't seem to get logged on failure.
+		return;
 	}
-	else
+
+	FGuardedInt64 GuardedEndOffset = FGuardedInt64(ctx->ReadOffset) + length;
+	if (GuardedEndOffset.InvalidOrGreaterThan(ctx->CompressedData.Num()))
 	{
-		ctx->SetError(TEXT("Invalid read position for CompressedData."));
+		UE_LOG(LogImageWrapper, Warning, TEXT("Bad PNG read position: offset %d num %lld length: %llu"), ctx->ReadOffset, ctx->CompressedData.Num(), length);
+		ctx->SetError(TEXT("Invalid read position for CompressedData.")); // this doesn't seem to get logged on failure.
+		return;
 	}
+
+	FMemory::Memcpy(data, &ctx->CompressedData[ctx->ReadOffset], length);
+	ctx->ReadOffset = GuardedEndOffset.Get(0);
 }
 
 
