EXR invalid file protection.

diff --git a/Engine/Source/Runtime/ImageWrapper/Private/Formats/ExrImageWrapper.cpp b/Engine/Source/Runtime/ImageWrapper/Private/Formats/ExrImageWrapper.cpp
--- a/Engine/Source/Runtime/ImageWrapper/Private/Formats/ExrImageWrapper.cpp
+++ b/Engine/Source/Runtime/ImageWrapper/Private/Formats/ExrImageWrapper.cpp
@@ -2,6 +2,7 @@
 
 #include "Formats/ExrImageWrapper.h"
 #include "ImageWrapperPrivate.h"
+#include "Math/GuardedInt.h"
 
 #include "ColorSpace.h"
 #include "Containers/StringConv.h"
@@ -90,6 +91,10 @@ public:
 		, Size(InSize)
 		, Pos(0)
 	{
+		if (InSize < 0)
+		{
+			UE_LOG(LogImageWrapper, Fatal, TEXT("Negative size passed to EXR parser, can not continue."));
+		}
 	}
 
 	//------------------------------------------------------
@@ -107,8 +112,12 @@ public:
 	{
 		int64 SrcN = InN;
 
-		if(Pos + SrcN > Size)
+		FGuardedInt64 NextPosition = FGuardedInt64(Pos) + SrcN;
+		if (SrcN < 0 ||
+			NextPosition.InvalidOrGreaterThan(Size))
 		{
+			// This is supposed to throw an exception per the spec?
+			throw std::overflow_error("invalid read size");
 			return false;
 		}
 
@@ -145,8 +154,8 @@ public:
 private:
 
 	const char* Data;
-	int64 Size;
-	int64 Pos;
+	uint64 Size;
+	uint64 Pos;
 };
 
 // these are the channel names we write
