fix crash in TiffImageWrapper, Tasks were still running after "this" was deleted

#rnx

diff --git a/Engine/Source/Runtime/ImageWrapper/Private/Formats/TiffImageWrapper.cpp b/Engine/Source/Runtime/ImageWrapper/Private/Formats/TiffImageWrapper.cpp
--- a/Engine/Source/Runtime/ImageWrapper/Private/Formats/TiffImageWrapper.cpp
+++ b/Engine/Source/Runtime/ImageWrapper/Private/Formats/TiffImageWrapper.cpp
@@ -993,6 +993,9 @@ namespace UE::ImageWrapper::Private
 
 		TQueue<TSharedPtr<TArray64<uint8>, ESPMode::ThreadSafe>, EQueueMode::Mpsc> UsableBufferQueue;
 		FGraphEventArray Tasks;
+		
+		ENamedThreads::Type NamedThread = IsInGameThread() ? ENamedThreads::GameThread : ENamedThreads::AnyThread;
+		ON_SCOPE_EXIT {	FTaskGraphInterface::Get().WaitUntilTasksComplete(Tasks, NamedThread); };
 
 		if constexpr (bIsTiled)
 		{
@@ -1112,9 +1115,20 @@ namespace UE::ImageWrapper::Private
 		// Write the alpha
 		if (bAddAlpha)
 		{
+			//Tasks[] is filling WriteArray and not done yet
+			//	even if Tasks only writes RGB and we only write A , so there is no race
+			//	it's better to not have them running at the same time
+			FTaskGraphInterface::Get().WaitUntilTasksComplete(Tasks, NamedThread);
+			Tasks.SetNum(0); // clear array so scope exit doesn't wait on tasks again
+
+			// this branch is only used for 3-channel RGB tiffs and 4-channel RGBA output pixels (NumOfChannelDest==4)
+
+			//todo: would be better to do this in the ProcessDecodedData tasks, so that we write the memory only once
+
+			//todo: ParallelFor over individual pixels is not great, better to do on rows; use ImageCore ImageParallelFor
 			ParallelFor(Width * Height, [NumOfChannelDest, &WriteArray](int32 Index)
 				{
-					const int64 WriteIndex = int64(NumOfChannelDest) * Index + 3;
+					const int64 WriteIndex = int64(NumOfChannelDest) * Index + NumOfChannelDest-1;
 					if constexpr (std::is_same_v<DataTypeDest, FFloat16> || TIsFloatingPoint<DataTypeDest>::Value)
 					{
 						WriteArray[WriteIndex] = DataTypeDest(1.f);
@@ -1126,9 +1140,7 @@ namespace UE::ImageWrapper::Private
 				}, IsInGameThread() ? EParallelForFlags::None : EParallelForFlags::BackgroundPriority);
 		}
 
-		ENamedThreads::Type NamedThread = IsInGameThread() ? ENamedThreads::GameThread : ENamedThreads::AnyThread;
-		FTaskGraphInterface::Get().WaitUntilTasksComplete(Tasks, NamedThread);
-
+		// scope exit will wait on Tasks[] completing
 
 		return true;
 	}
