[Backout] - CL25715182
Original CL Desc
-----------------------------------------------------------------
EXR invalid file protection.

diff --git a/Engine/Source/Runtime/ImageWrapper/Private/Formats/ExrImageWrapper.cpp b/Engine/Source/Runtime/ImageWrapper/Private/Formats/ExrImageWrapper.cpp
--- a/Engine/Source/Runtime/ImageWrapper/Private/Formats/ExrImageWrapper.cpp
+++ b/Engine/Source/Runtime/ImageWrapper/Private/Formats/ExrImageWrapper.cpp
@@ -2,7 +2,6 @@
 
 #include "Formats/ExrImageWrapper.h"
 #include "ImageWrapperPrivate.h"
-#include "Math/GuardedInt.h"
 
 #include "ColorSpace.h"
 #include "Containers/StringConv.h"
@@ -91,10 +90,6 @@ public:
 		, Size(InSize)
 		, Pos(0)
 	{
-		if (InSize < 0)
-		{
-			UE_LOG(LogImageWrapper, Fatal, TEXT("Negative size passed to EXR parser, can not continue."));
-		}
 	}
 
 	//------------------------------------------------------
@@ -112,12 +107,8 @@ public:
 	{
 		int64 SrcN = InN;
 
-		FGuardedInt64 NextPosition = FGuardedInt64(Pos) + SrcN;
-		if (SrcN < 0 ||
-			NextPosition.InvalidOrGreaterThan(Size))
+		if(Pos + SrcN > Size)
 		{
-			// This is supposed to throw an exception per the spec?
-			throw std::overflow_error("invalid read size");
 			return false;
 		}
 
@@ -154,8 +145,8 @@ public:
 private:
 
 	const char* Data;
-	uint64 Size;
-	uint64 Pos;
+	int64 Size;
+	int64 Pos;
 };
 
 // these are the channel names we write
