fix image import crashes in Ico and Tga

#rnx

diff --git a/Engine/Source/Runtime/ImageWrapper/Private/Formats/IcoImageWrapper.cpp b/Engine/Source/Runtime/ImageWrapper/Private/Formats/IcoImageWrapper.cpp
--- a/Engine/Source/Runtime/ImageWrapper/Private/Formats/IcoImageWrapper.cpp
+++ b/Engine/Source/Runtime/ImageWrapper/Private/Formats/IcoImageWrapper.cpp
@@ -90,7 +90,7 @@ void FIcoImageWrapper::Uncompress( const ERGBFormat InFormat, const int32 InBitD
 {
 	const uint8* Buffer = CompressedData.GetData();
 
-	if (ImageOffset != 0 && ImageSize != 0)
+	if (ImageOffset != 0 && ImageSize != 0 && SubImageWrapper != nullptr)
 	{
 		SubImageWrapper->Uncompress(InFormat, InBitDepth);
 		// Uncompress has no return value
@@ -113,7 +113,7 @@ bool FIcoImageWrapper::GetRaw( const ERGBFormat InFormat, int32 InBitDepth, TArr
 	LastError.Empty();
 	Uncompress(InFormat, InBitDepth);
 
-	if (LastError.IsEmpty())
+	if (LastError.IsEmpty() && SubImageWrapper != nullptr)
 	{
 		SubImageWrapper->MoveRawData(OutRawData);
 		if ( OutRawData.IsEmpty() )
@@ -132,6 +132,7 @@ bool FIcoImageWrapper::GetRaw( const ERGBFormat InFormat, int32 InBitDepth, TArr
 bool FIcoImageWrapper::LoadICOHeader()
 {
 	const uint8* Buffer = CompressedData.GetData();
+	uint64 BufferSize = CompressedData.Num();
 
 #if WITH_UNREALPNG
 	TSharedPtr<FPngImageWrapper> PngWrapper = MakeShareable(new FPngImageWrapper);
@@ -139,6 +140,13 @@ bool FIcoImageWrapper::LoadICOHeader()
 	TSharedPtr<FBmpImageWrapper> BmpWrapper = MakeShareable(new FBmpImageWrapper(false, true));
 
 	bool bFoundImage = false;
+	SubImageWrapper = nullptr;
+
+	if ( BufferSize < sizeof(FIconDir) )
+	{
+		SetError(TEXT("LoadICOHeader corrupt; insufficient data"));
+		return false;
+	}
 	const FIconDir* IconHeader = (FIconDir*)(Buffer);
 	
 	if (IconHeader->idReserved == 0 && IconHeader->idType == 1)
@@ -147,13 +155,31 @@ bool FIcoImageWrapper::LoadICOHeader()
 		uint32 LargestWidth = 0;
 		const FIconDirEntry* IconDirEntry = IconHeader->idEntries;
 		
+		if ( IconHeader->idCount == 0 )
+		{
+			SetError(TEXT("LoadICOHeader corrupt; no entries"));
+			return false;
+		}
+		// FIconDir has 1 entry in it
+		if ( BufferSize < sizeof(FIconDir) + (IconHeader->idCount-1)*sizeof(FIconDirEntry) )
+		{
+			SetError(TEXT("LoadICOHeader corrupt; insufficient data"));
+			return false;
+		}
+
 		for (int32 Entry = 0; Entry < (int32)IconHeader->idCount; Entry++, IconDirEntry++)
 		{
 			const uint32 RealWidth = IconDirEntry->bWidth == 0 ? 256 : IconDirEntry->bWidth;
 			if ( IconDirEntry->wBitCount == 32 && RealWidth > LargestWidth )
 			{
+				if ( (uint64)IconDirEntry->dwImageOffset + IconDirEntry->dwBytesInRes > BufferSize )
+				{
+					SetError(TEXT("LoadICOHeader corrupt; insufficient data"));
+					return false;
+				}
+
 #if WITH_UNREALPNG
-				if (PngWrapper->SetCompressed(Buffer + IconDirEntry->dwImageOffset, (int32)IconDirEntry->dwBytesInRes))
+				if (PngWrapper->SetCompressed(Buffer + IconDirEntry->dwImageOffset, (int64)IconDirEntry->dwBytesInRes))
 				{
 					Width = PngWrapper->GetWidth();
 					Height = PngWrapper->GetHeight();
@@ -167,7 +193,7 @@ bool FIcoImageWrapper::LoadICOHeader()
 				}
 				else 
 #endif
-				if (BmpWrapper->SetCompressed(Buffer + IconDirEntry->dwImageOffset, (int32)IconDirEntry->dwBytesInRes))
+				if (BmpWrapper->SetCompressed(Buffer + IconDirEntry->dwImageOffset, (int64)IconDirEntry->dwBytesInRes))
 				{
 					// otherwise this should be a BMP icon
 					Width = BmpWrapper->GetWidth();
diff --git a/Engine/Source/Runtime/ImageWrapper/Private/TgaImageSupport.cpp b/Engine/Source/Runtime/ImageWrapper/Private/TgaImageSupport.cpp
--- a/Engine/Source/Runtime/ImageWrapper/Private/TgaImageSupport.cpp
+++ b/Engine/Source/Runtime/ImageWrapper/Private/TgaImageSupport.cpp
@@ -9,7 +9,7 @@ namespace TgaImageSupportImpl
 	{
 		const uint8* const IdData = (uint8*)TGA + sizeof(FTGAFileHeader); 
 		const uint8* const  ColorMap = IdData + TGA->IdFieldLength;
-		const uint8* ImageData = (uint8*) (ColorMap + (TGA->ColorMapEntrySize + 4) / 8 * TGA->ColorMapLength);
+		const uint8* ImageData = (uint8*) (ColorMap + (TGA->ColorMapEntrySize + 4) / 8 * (int64)TGA->ColorMapLength);
 
 		const uint8* TGAEnd = (uint8*)TGA + TGABufferLenght;
 
@@ -17,9 +17,9 @@ namespace TgaImageSupportImpl
 		int32 RLERun = 0;
 		int32 RAWRun = 0;
 
-		for(int32 Y = TGA->Height-1; Y >=0; Y--) // Y-flipped.
+		for(int64 Y = TGA->Height-1; Y >=0; Y--) // Y-flipped.
 		{
-			for(int32 X = 0;X < TGA->Width;X++)
+			for(int64 X = 0;X < TGA->Width;X++)
 			{
 				if( RLERun > 0 )
 				{
@@ -27,7 +27,7 @@ namespace TgaImageSupportImpl
 				}
 				else if( RAWRun == 0 ) // new raw pixel or RLE-run.
 				{
-					if ( TGAEnd < ImageData )
+					if ( TGAEnd <= ImageData )
 					{
 						return false;
 					}
@@ -67,7 +67,7 @@ namespace TgaImageSupportImpl
 	{
 		const uint8* const IdData = (uint8*)TGA + sizeof(FTGAFileHeader); 
 		const uint8* const ColorMap = IdData + TGA->IdFieldLength;
-		const uint8* ImageData = (uint8*) (ColorMap + (TGA->ColorMapEntrySize + 4) / 8 * TGA->ColorMapLength);
+		const uint8* ImageData = (uint8*) (ColorMap + (TGA->ColorMapEntrySize + 4) / 8 * (int64)TGA->ColorMapLength);
 
 		const uint8* TGAEnd = (uint8*)TGA + TGABufferLenght;
 
@@ -75,9 +75,9 @@ namespace TgaImageSupportImpl
 		int32 RLERun = 0;
 		int32 RAWRun = 0;
 
-		for(int32 Y = TGA->Height-1; Y >=0; Y--) // Y-flipped.
+		for(int64 Y = TGA->Height-1; Y >=0; Y--) // Y-flipped.
 		{
-			for(int32 X = 0;X < TGA->Width;X++)
+			for(int64 X = 0;X < TGA->Width;X++)
 			{
 				if( RLERun > 0 )
 				{
@@ -85,7 +85,7 @@ namespace TgaImageSupportImpl
 				}
 				else if( RAWRun == 0 ) // new raw pixel or RLE-run.
 				{
-					if ( TGAEnd < ImageData )
+					if ( TGAEnd <= ImageData )
 					{
 						return false;
 					}
@@ -126,7 +126,7 @@ namespace TgaImageSupportImpl
 	{
 		const uint8* const IdData = (uint8*)TGA + sizeof(FTGAFileHeader);
 		const uint8* const ColorMap = IdData + TGA->IdFieldLength;				
-		const uint16* ImageData = (uint16*) (ColorMap + (TGA->ColorMapEntrySize + 4) / 8 * TGA->ColorMapLength);
+		const uint16* ImageData = (uint16*) (ColorMap + (TGA->ColorMapEntrySize + 4) / 8 * (int64)TGA->ColorMapLength);
 
 		const uint16* TGAEnd = (uint16*)TGA + TGABufferLenght / 2;
 
@@ -134,9 +134,9 @@ namespace TgaImageSupportImpl
 		int32 RLERun = 0;
 		int32 RAWRun = 0;
 
-		for(int32 Y = TGA->Height-1; Y >=0; Y--) // Y-flipped.
+		for(int64 Y = TGA->Height-1; Y >=0; Y--) // Y-flipped.
 		{					
-			for( int32 X=0;X<TGA->Width;X++ )
+			for( int64 X=0;X<TGA->Width;X++ )
 			{						
 				if( RLERun > 0 )
 				{
@@ -144,7 +144,7 @@ namespace TgaImageSupportImpl
 				}
 				else if( RAWRun == 0 ) // new raw pixel or RLE-run.
 				{
-					if ( TGAEnd < ImageData )
+					if ( TGAEnd <= ImageData )
 					{
 						return false;
 					}
@@ -190,14 +190,14 @@ namespace TgaImageSupportImpl
 	{
 		const uint8* const IdData = (uint8*)TGA + sizeof(FTGAFileHeader);
 		const uint8* const ColorMap = IdData + TGA->IdFieldLength;
-		const uint32* const ImageData = (uint32*) (ColorMap + (TGA->ColorMapEntrySize + 4) / 8 * TGA->ColorMapLength);
+		const uint32* const ImageData = (uint32*) (ColorMap + (TGA->ColorMapEntrySize + 4) / 8 * (int64)TGA->ColorMapLength);
 
-		if ( (uint8*)TGA + TGABufferLenght < (uint8*)(ImageData + TGA->Width * TGA->Height) )
+		if ( (uint8*)TGA + TGABufferLenght < (uint8*)(ImageData + (int64)TGA->Width * TGA->Height) )
 		{
 			return false;
 		}
 
-		for(int32 Y = 0;Y < TGA->Height;Y++)
+		for(int64 Y = 0;Y < TGA->Height;Y++)
 		{
 			FMemory::Memcpy(TextureData + Y * TGA->Width,ImageData + (TGA->Height - Y - 1) * TGA->Width,TGA->Width * 4);
 		}
@@ -209,17 +209,17 @@ namespace TgaImageSupportImpl
 	{
 		const uint8* const IdData = (uint8*)TGA + sizeof(FTGAFileHeader);
 		const uint8* const ColorMap = IdData + TGA->IdFieldLength;
-		const uint8* const ImageData = (uint8*) (ColorMap + (TGA->ColorMapEntrySize + 4) / 8 * TGA->ColorMapLength);
+		const uint8* const ImageData = (uint8*) (ColorMap + (TGA->ColorMapEntrySize + 4) / 8 * (int64)TGA->ColorMapLength);
 		uint8 Pixel[4];
 
-		if ( (uint8*)TGA + TGABufferLenght < ImageData + TGA->Width * TGA->Height * 3 )
+		if ( (uint8*)TGA + TGABufferLenght < ImageData + (int64)TGA->Width * TGA->Height * 3 )
 		{
 			return false;
 		}
 
-		for(int32 Y = 0; Y < TGA->Height; Y++)
+		for(int64 Y = 0; Y < TGA->Height; Y++)
 		{
-			for(int32 X = 0; X < TGA->Width; X++)
+			for(int64 X = 0; X < TGA->Width; X++)
 			{
 				Pixel[0] = *(( ImageData+( TGA->Height-Y-1 )*TGA->Width*3 )+X*3+0);
 				Pixel[1] = *(( ImageData+( TGA->Height-Y-1 )*TGA->Width*3 )+X*3+1);
@@ -234,20 +234,21 @@ namespace TgaImageSupportImpl
 
 	bool DecompressTGA_16bpp( const FTGAFileHeader* TGA, const int64 TGABufferLenght, uint32* TextureData )
 	{
+		//todo: fix : this kind of code is prone to overflow bugs :
 		const uint8* const IdData = (uint8*)TGA + sizeof(FTGAFileHeader);
 		const uint8* const ColorMap = IdData + TGA->IdFieldLength;
-		const uint16* ImageData = (uint16*) (ColorMap + (TGA->ColorMapEntrySize + 4) / 8 * TGA->ColorMapLength);
+		const uint16* ImageData = (uint16*) (ColorMap + (TGA->ColorMapEntrySize + 4) / 8 * (int64)TGA->ColorMapLength);
 		uint16 FilePixel = 0;
 		uint32 TexturePixel = 0;
 
-		if ( (uint16*)((uint8*)TGA + TGABufferLenght) < ImageData + TGA->Height * TGA->Width )
+		if ( (uint16*)((uint8*)TGA + TGABufferLenght) < ImageData + (int64)TGA->Height * TGA->Width )
 		{
 			return false;
 		}
 
-		for (int32 Y = TGA->Height - 1; Y >= 0; Y--)
+		for (int64 Y = TGA->Height - 1; Y >= 0; Y--)
 		{					
-			for (int32 X = 0; X < TGA->Width; X++)
+			for (int64 X = 0; X < TGA->Width; X++)
 			{
 				FilePixel = *ImageData++;
 				// Convert file format A1R5G5B5 into pixel format B8G8R8A8
@@ -267,15 +268,15 @@ namespace TgaImageSupportImpl
 	{
 		const uint8* const IdData = (uint8*)TGA + sizeof(FTGAFileHeader);
 		const uint8* const  ColorMap = IdData + TGA->IdFieldLength;
-		const uint8* const  ImageData = (uint8*) (ColorMap + (TGA->ColorMapEntrySize + 4) / 8 * TGA->ColorMapLength);
+		const uint8* const  ImageData = (uint8*) (ColorMap + (TGA->ColorMapEntrySize + 4) / 8 * (int64)TGA->ColorMapLength);
 
-		if ( (uint8*)TGA + TGABufferLenght < ImageData + TGA->Width * TGA->Height )
+		if ( (uint8*)TGA + TGABufferLenght < ImageData + (int64)TGA->Width * TGA->Height )
 		{
 			return false;
 		}
 
-		int32 RevY = 0;
-		for (int32 Y = TGA->Height-1; Y >= 0; --Y)
+		int64 RevY = 0;
+		for (int64 Y = TGA->Height-1; Y >= 0; --Y)
 		{
 			const uint8* ImageCol = ImageData + (Y * TGA->Width); 
 			uint8* TextureCol = TextureData + (RevY++ * TGA->Width);
@@ -286,7 +287,7 @@ namespace TgaImageSupportImpl
 	}
 }
 
-bool DecompressTGA_helper( const FTGAFileHeader* TgaHeader, const int64 TGABufferLenght, uint32*& TextureData, const int32 TextureDataSize )
+bool DecompressTGA_helper( const FTGAFileHeader* TgaHeader, const int64 TGABufferLenght, uint32*& TextureData, const int64 TextureDataSize )
 {
 	bool bSuccess = false;
 	if ( TgaHeader->ImageTypeCode == 10 ) // 10 = RLE compressed 
@@ -315,7 +316,7 @@ bool DecompressTGA_helper( const FTGAFileHeader* TgaHeader, const int64 TGABuffe
 	}
 	else if(TgaHeader->ImageTypeCode == 2) // 2 = Uncompressed RGB
 	{
-		check( TextureDataSize == TgaHeader->Width * TgaHeader->Height * 4 );
+		check( TextureDataSize == (int64)TgaHeader->Width * TgaHeader->Height * 4 );
 
 		if(TgaHeader->BitsPerPixel == 32)
 		{
@@ -369,20 +370,20 @@ bool DecompressTGA_helper( const FTGAFileHeader* TgaHeader, const int64 TGABuffe
 		TArray<uint8> FlippedData;
 		FlippedData.AddUninitialized(TextureDataSize);
 
-		int32 NumBlocksX = TgaHeader->Width;
-		int32 NumBlocksY = TgaHeader->Height;
-		int32 BlockBytes = TgaHeader->BitsPerPixel == 8 ? 1 : 4;
+		int64 NumBlocksX = TgaHeader->Width;
+		int64 NumBlocksY = TgaHeader->Height;
+		int64 BlockBytes = TgaHeader->BitsPerPixel == 8 ? 1 : 4;
 		
 		check( TextureDataSize == NumBlocksX * NumBlocksY * BlockBytes );
 
 		uint8* MipData = (uint8*)TextureData;
 
-		for( int32 Y = 0; Y < NumBlocksY;Y++ )
+		for( int64 Y = 0; Y < NumBlocksY;Y++ )
 		{
-			for( int32 X  = 0; X < NumBlocksX; X++ )
+			for( int64 X  = 0; X < NumBlocksX; X++ )
 			{
-				int32 DestX = bFlipX ? (NumBlocksX - X - 1) : X;
-				int32 DestY = bFlipY ? (NumBlocksY - Y - 1) : Y;
+				int64 DestX = bFlipX ? (NumBlocksX - X - 1) : X;
+				int64 DestY = bFlipY ? (NumBlocksY - Y - 1) : Y;
 				FMemory::Memcpy(
 					&FlippedData[(DestX + DestY * NumBlocksX) * BlockBytes],
 					&MipData[(X + Y * NumBlocksX) * BlockBytes],
diff --git a/Engine/Source/Runtime/ImageWrapper/Public/TgaImageSupport.h b/Engine/Source/Runtime/ImageWrapper/Public/TgaImageSupport.h
--- a/Engine/Source/Runtime/ImageWrapper/Public/TgaImageSupport.h
+++ b/Engine/Source/Runtime/ImageWrapper/Public/TgaImageSupport.h
@@ -37,4 +37,4 @@ struct FTGAFileHeader
  * This helper allows to decompress TGA data in a pre-allocated memory block.
  * The pixel format is necessarily PF_A8R8G8B8.
  */
-IMAGEWRAPPER_API bool DecompressTGA_helper(	const FTGAFileHeader* TGA, const int64 TGABufferLenght, uint32*& TextureData, const int32 TextureDataSize );
+IMAGEWRAPPER_API bool DecompressTGA_helper(	const FTGAFileHeader* TGA, const int64 TGABufferLength, uint32*& TextureData, const int64 TextureDataSize );
