ImageWrapperModule StressDecompressImage for automated corrupt image import testing

#rnx

diff --git a/Engine/Source/Runtime/ImageWrapper/Private/Formats/PngImageWrapper.cpp b/Engine/Source/Runtime/ImageWrapper/Private/Formats/PngImageWrapper.cpp
--- a/Engine/Source/Runtime/ImageWrapper/Private/Formats/PngImageWrapper.cpp
+++ b/Engine/Source/Runtime/ImageWrapper/Private/Formats/PngImageWrapper.cpp
@@ -499,7 +499,7 @@ void FPngImageWrapper::UncompressPNGData(const ERGBFormat InFormat, const int32
 		 *	an unhandled exception upon a CRC error. This code 
 		 *	catches our custom exception thrown in user_error_fn.
 		 */
-		UE_LOG(LogImageWrapper, Error, TEXT("%s"), *e.ErrorText);
+		UE_LOG(LogImageWrapper, Error, TEXT("FPNGImageCRCError: %s"), *e.ErrorText);
 	}
 #endif
 
@@ -560,6 +560,10 @@ bool FPngImageWrapper::LoadPNGHeader()
 
 		// ---------------------------------------------------------------------------------------------------------
 		// Anything allocated on the stack after this point will not be destructed correctly in the case of an error
+		
+#if !PLATFORM_EXCEPTIONS_DISABLED
+		try
+#endif
 		{
 			png_set_read_fn(png_ptr, this, FPngImageWrapper::user_read_compressed);
 
@@ -584,6 +588,18 @@ bool FPngImageWrapper::LoadPNGHeader()
 				Format = ERGBFormat::BGRA;
 			}
 		}
+#if !PLATFORM_EXCEPTIONS_DISABLED
+		catch (const FPNGImageCRCError& e)
+		{
+			/** 
+			 *	libPNG has a known issue in version 1.5.2 causing
+			 *	an unhandled exception upon a CRC error. This code 
+			 *	catches our custom exception thrown in user_error_fn.
+			 */
+			UE_LOG(LogImageWrapper, Error, TEXT("FPNGImageCRCError: %s"), *e.ErrorText);
+			return false;
+		}
+#endif
 
 		return true;
 	}
diff --git a/Engine/Source/Runtime/ImageWrapper/Private/ImageWrapperModule.cpp b/Engine/Source/Runtime/ImageWrapper/Private/ImageWrapperModule.cpp
--- a/Engine/Source/Runtime/ImageWrapper/Private/ImageWrapperModule.cpp
+++ b/Engine/Source/Runtime/ImageWrapper/Private/ImageWrapperModule.cpp
@@ -185,14 +185,17 @@ public:
 		{
 			Format = EImageFormat::TIFF;
 		}
-		else if ( FTgaImageWrapper::IsTGAHeader(CompressedData, CompressedSize) )
-		{
-			Format = EImageFormat::TGA;
-		}
 		else if ( UE::DDS::FDDSFile::IsADDS((const uint8 *)CompressedData, CompressedSize) )
 		{
 			Format = EImageFormat::DDS;
 		}
+		else if ( FTgaImageWrapper::IsTGAHeader(CompressedData, CompressedSize) )
+		{
+			// beware: TGA does not have any type Id in the header
+			//	this can mis-identify binaries as TGA
+			//  it should be the last format checked in the list
+			Format = EImageFormat::TGA;
+		}
 
 		return Format;
 	}
@@ -310,6 +313,53 @@ public:
 		}
 	}
 	
+	// Stress decoding an image many times with corruption
+	//	to test if the ImageWrapper handles invalid data
+	// this is by no means a substitute for real systematic fuzz testing, but is better than nothing
+	void StressDecompressImage(TSharedPtr<IImageWrapper> ImageWrapper,const void* InCompressedData, int64 InCompressedSize)
+	{
+		FImage OutImage;
+
+		// alloc a large aligned buf (would prefer to use VirtualAlloc directly so we get no-access pages adjacent)
+		int64 AllocSize = (InCompressedSize + 0xFFFF) & (~0xFFFFULL);
+		uint8 * Scratch = (uint8 *) FMemory::Malloc(AllocSize,65536);
+		uint8 * End = Scratch + AllocSize;
+
+		// test truncations of the stream :
+		//for(int64 TruncatedSize = 1;TruncatedSize < InCompressedSize;TruncatedSize++)
+		for(int64 TruncatedSize = InCompressedSize-1, Step=1; TruncatedSize>0; TruncatedSize -= Step, Step += Step+1)
+		{
+			// put truncated data at end of allocation to test reading past end :
+			uint8 * TruncatedData = End - TruncatedSize;
+			memcpy(TruncatedData,InCompressedData,TruncatedSize);
+			
+			if ( ImageWrapper->SetCompressed(TruncatedData,TruncatedSize) )
+			{
+				ImageWrapper->GetRawImage(OutImage);
+			}
+		}
+		
+		memcpy(Scratch,InCompressedData,InCompressedSize);
+
+		// toggle one bit :
+		for(int Bit=1; Bit < 256; Bit *= 2)
+		{
+			for(int64 Pos=0, Step=1;Pos < InCompressedSize;Pos+= Step, Step += Step+1)
+			{
+				Scratch[Pos] ^= Bit;
+			
+				if ( ImageWrapper->SetCompressed(Scratch,InCompressedSize) )
+				{
+					ImageWrapper->GetRawImage(OutImage);
+				}
+
+				Scratch[Pos] ^= Bit;
+			}
+		}
+
+		FMemory::Free(Scratch);
+	}
+
 	virtual bool DecompressImage(const void* InCompressedData, int64 InCompressedSize, FImage & OutImage) override
 	{
 		TRACE_CPUPROFILER_EVENT_SCOPE(ImageWrapper.Decompress);
@@ -325,6 +375,11 @@ public:
 		{
 			return false;
 		}
+		
+		#if 0
+		// for testing :
+		StressDecompressImage(ImageWrapper,InCompressedData,InCompressedSize);
+		#endif
 
 		if ( ! ImageWrapper->SetCompressed(InCompressedData,InCompressedSize) )
 		{
