ExrImageWrapper : corrupt import should not be fatal, just an error
use runtime_exception for throws
fix return value of read function, though it seems not to be used

#rnx

diff --git a/Engine/Source/Runtime/ImageWrapper/Private/Formats/ExrImageWrapper.cpp b/Engine/Source/Runtime/ImageWrapper/Private/Formats/ExrImageWrapper.cpp
--- a/Engine/Source/Runtime/ImageWrapper/Private/Formats/ExrImageWrapper.cpp
+++ b/Engine/Source/Runtime/ImageWrapper/Private/Formats/ExrImageWrapper.cpp
@@ -11,6 +11,8 @@
 
 #if WITH_UNREALEXR
 
+#include <stdexcept>
+
 FExrImageWrapper::FExrImageWrapper()
 	: FImageWrapperBase()
 {
@@ -91,9 +93,11 @@ public:
 		, Size(InSize)
 		, Pos(0)
 	{
+		// InSize is signed but "Size" member is unsigned
 		if (InSize < 0)
 		{
-			UE_LOG(LogImageWrapper, Fatal, TEXT("Negative size passed to EXR parser, can not continue."));
+			Size = 0;
+			throw std::runtime_error("FMemFileIn: Negative size passed to EXR parser");
 		}
 	}
 
@@ -108,28 +112,30 @@ public:
 	//------------------------------------------------------
 
 	// InN must be 32bit to match the abstract interface.
-	virtual bool read (char c[/*n*/], int32 InN)
+	virtual bool read (char Out[/*n*/], int32 Count)
 	{
-		int64 SrcN = InN;
+		// return false if EOF is hit
+		// OpenEXR mostly ignores this return value, you must throw to get error handling
 
-		FGuardedInt64 NextPosition = FGuardedInt64(Pos) + SrcN;
-		if (SrcN < 0 ||
+		FGuardedInt64 NextPosition = FGuardedInt64(Pos) + Count;
+		if (Count < 0 ||
 			NextPosition.InvalidOrGreaterThan(Size))
 		{
-			// This is supposed to throw an exception per the spec, however we can't use
-			// an actual exception type right now due to linker issues during LTCG.
-			UE_LOG(LogImageWrapper, Error, TEXT("Invalid read requested in EXR"));
-			throw std::exception();
+			throw std::runtime_error("FMemFileIn: Exr read out of bounds");
 			return false;
 		}
 
-		for (int64 i = 0; i < SrcN; ++i)
+		memcpy(Out,Data+Pos,Count);
+		Pos += Count; // == NextPosition
+
+		if ( Pos == Size )
+		{
+			return false;
+		}
+		else
 		{
-			c[i] = Data[Pos];
-			++Pos;
+			return true;
 		}
-
-		return Pos >= Size;
 	}
 
 	//--------------------------------------------------------
