Tga image importer cleanups and robustination

#rnx

diff --git a/Engine/Source/Runtime/ImageWrapper/Private/TgaImageSupport.cpp b/Engine/Source/Runtime/ImageWrapper/Private/TgaImageSupport.cpp
--- a/Engine/Source/Runtime/ImageWrapper/Private/TgaImageSupport.cpp
+++ b/Engine/Source/Runtime/ImageWrapper/Private/TgaImageSupport.cpp
@@ -5,13 +5,55 @@
 
 namespace TgaImageSupportImpl
 {
-	bool DecompressTGA_RLE_32bpp( const FTGAFileHeader* TGA, const int64 TGABufferLenght, uint32* TextureData )
+	
+	static const uint8 * GetImageData(const FTGAFileHeader* TGA, const int64 TGABufferLength)
 	{
-		const uint8* const IdData = (uint8*)TGA + sizeof(FTGAFileHeader); 
-		const uint8* const  ColorMap = IdData + TGA->IdFieldLength;
-		const uint8* ImageData = (uint8*) (ColorMap + (TGA->ColorMapEntrySize + 4) / 8 * (int64)TGA->ColorMapLength);
+		if ( TGABufferLength < sizeof(FTGAFileHeader) )
+			return nullptr;
 
-		const uint8* TGAEnd = (uint8*)TGA + TGABufferLenght;
+		int64 ColorMapEntryBytes = (TGA->ColorMapEntrySize + 7) / 8;
+		int64 ColorMapBytes = ColorMapEntryBytes * TGA->ColorMapLength;
+		
+		int64 OffsetToImageData = sizeof(FTGAFileHeader) + TGA->IdFieldLength + ColorMapBytes;
+
+		if ( TGABufferLength < OffsetToImageData )
+			return nullptr;
+
+		const uint8* ImageData = (const uint8*)TGA + OffsetToImageData;
+		
+		return ImageData;
+	}
+	
+	static uint32 ConvertA1R5G5B5ToBGRA8(uint32 FilePixel)
+	{
+		// Convert file format A1R5G5B5 into pixel format B8G8R8A8
+		uint32 TexturePixel;
+		TexturePixel  = ((FilePixel & 0x001F) * 33) >> 2;
+		TexturePixel += ((FilePixel & 0x03E0) * (33 << 1)) & 0x0000FF00;
+		TexturePixel += ((FilePixel & 0x7C00) * (33 << 4)) & 0x00FF0000;
+		TexturePixel += (FilePixel & 0x8000) * (255<<9);
+		return TexturePixel;
+	}
+	
+	static uint32 ConvertA1R5G5B5ToBGRA8(const uint8 * Bytes)
+	{
+		uint32 FilePixel = Bytes[0] + (Bytes[1]<<8);
+		return ConvertA1R5G5B5ToBGRA8(FilePixel);
+	}
+
+	static uint32 ConvertBGRToBGRA8(const uint8 * BGR)
+	{
+		uint32 Pixel = BGR[0] + (BGR[1]<<8) + (BGR[2]<<16) + 0xFF000000U;
+		return Pixel;
+	}
+
+	bool DecompressTGA_RLE_32bpp( const FTGAFileHeader* TGA, const int64 TGABufferLength, uint32* TextureData )
+	{
+		const uint8* ImageData = GetImageData(TGA,TGABufferLength);
+		if ( ImageData == nullptr )
+			return false;
+
+		const uint8* TGAEnd = (uint8*)TGA + TGABufferLength;
 
 		uint32 Pixel = 0;
 		int32 RLERun = 0;
@@ -63,15 +105,15 @@ namespace TgaImageSupportImpl
 		return true;
 	}
 
-	bool DecompressTGA_RLE_24bpp( const FTGAFileHeader* TGA, const int64 TGABufferLenght, uint32* TextureData )
+	bool DecompressTGA_RLE_24bpp( const FTGAFileHeader* TGA, const int64 TGABufferLength, uint32* TextureData )
 	{
-		const uint8* const IdData = (uint8*)TGA + sizeof(FTGAFileHeader); 
-		const uint8* const ColorMap = IdData + TGA->IdFieldLength;
-		const uint8* ImageData = (uint8*) (ColorMap + (TGA->ColorMapEntrySize + 4) / 8 * (int64)TGA->ColorMapLength);
+		const uint8* ImageData = GetImageData(TGA,TGABufferLength);
+		if ( ImageData == nullptr )
+			return false;
 
-		const uint8* TGAEnd = (uint8*)TGA + TGABufferLenght;
+		const uint8* TGAEnd = (uint8*)TGA + TGABufferLength;
 
-		uint8 Pixel[4] = {};
+		uint32 Pixel = 0;
 		int32 RLERun = 0;
 		int32 RAWRun = 0;
 
@@ -108,27 +150,25 @@ namespace TgaImageSupportImpl
 				// Retrieve new pixel data - raw run or single pixel for RLE stretch.
 				if( RAWRun > 0 )
 				{
-					Pixel[0] = *(ImageData++);
-					Pixel[1] = *(ImageData++);
-					Pixel[2] = *(ImageData++);
-					Pixel[3] = 255;
+					Pixel = ConvertBGRToBGRA8(ImageData);
+					ImageData += 3;
 					RAWRun--;
 					RLERun--;
 				}
 				// Store.
-				*( (TextureData + Y*TGA->Width)+X ) = *(uint32*)&Pixel;
+				*( (TextureData + Y*TGA->Width)+X ) = Pixel;
 			}
 		}
 		return true;
 	}
 
-	bool DecompressTGA_RLE_16bpp( const FTGAFileHeader* TGA, const int64 TGABufferLenght, uint32* TextureData )
+	bool DecompressTGA_RLE_16bpp( const FTGAFileHeader* TGA, const int64 TGABufferLength, uint32* TextureData )
 	{
-		const uint8* const IdData = (uint8*)TGA + sizeof(FTGAFileHeader);
-		const uint8* const ColorMap = IdData + TGA->IdFieldLength;				
-		const uint16* ImageData = (uint16*) (ColorMap + (TGA->ColorMapEntrySize + 4) / 8 * (int64)TGA->ColorMapLength);
+		const uint8* ImageData = GetImageData(TGA,TGABufferLength);
+		if ( ImageData == nullptr )
+			return false;
 
-		const uint16* TGAEnd = (uint16*)TGA + TGABufferLenght / 2;
+		const uint8* TGAEnd = (uint8*)TGA + TGABufferLength;
 
 		uint32 TexturePixel = 0;
 		int32 RLERun = 0;
@@ -148,8 +188,7 @@ namespace TgaImageSupportImpl
 					{
 						return false;
 					}
-					uint8 RLEChunk =  *((uint8*)ImageData);
-					ImageData = (uint16*)(((uint8*)ImageData)+1);
+					uint8 RLEChunk =  *ImageData++;
 					if( RLEChunk & 0x80 )
 					{
 						RLERun = ( RLEChunk & 0x7F ) + 1;
@@ -160,7 +199,7 @@ namespace TgaImageSupportImpl
 						RAWRun = ( RLEChunk & 0x7F ) + 1;
 					}
 
-					if ( TGAEnd < ImageData + RAWRun )
+					if ( TGAEnd < ImageData + RAWRun*2 )
 					{
 						return false;
 					}
@@ -168,15 +207,10 @@ namespace TgaImageSupportImpl
 				// Retrieve new pixel data - raw run or single pixel for RLE stretch.
 				if( RAWRun > 0 )
 				{ 
-					const uint16 FilePixel = *(ImageData++);
+					TexturePixel = ConvertA1R5G5B5ToBGRA8(ImageData);
+					ImageData += 2;
 					RAWRun--;
 					RLERun--;
-
-					// Convert file format A1R5G5B5 into pixel format B8G8R8B8
-					TexturePixel = (FilePixel & 0x001F) << 3;
-					TexturePixel |= (FilePixel & 0x03E0) << 6;
-					TexturePixel |= (FilePixel & 0x7C00) << 9;
-					TexturePixel |= (FilePixel & 0x8000) << 16;
 				}
 				// Store.
 				*( (TextureData + Y*TGA->Width)+X ) = TexturePixel;
@@ -186,33 +220,32 @@ namespace TgaImageSupportImpl
 		return true;
 	}
 
-	bool DecompressTGA_32bpp( const FTGAFileHeader* TGA, const int64 TGABufferLenght, uint32* TextureData )
+	bool DecompressTGA_32bpp( const FTGAFileHeader* TGA, const int64 TGABufferLength, uint32* TextureData )
 	{
-		const uint8* const IdData = (uint8*)TGA + sizeof(FTGAFileHeader);
-		const uint8* const ColorMap = IdData + TGA->IdFieldLength;
-		const uint32* const ImageData = (uint32*) (ColorMap + (TGA->ColorMapEntrySize + 4) / 8 * (int64)TGA->ColorMapLength);
+		const uint8* ImageData = GetImageData(TGA,TGABufferLength);
+		if ( ImageData == nullptr )
+			return false;
 
-		if ( (uint8*)TGA + TGABufferLenght < (uint8*)(ImageData + (int64)TGA->Width * TGA->Height) )
+		if ( TGABufferLength < ((uint8*)ImageData - (uint8*)TGA) + (int64)TGA->Width * TGA->Height*4 )
 		{
 			return false;
 		}
 
 		for(int64 Y = 0;Y < TGA->Height;Y++)
 		{
-			FMemory::Memcpy(TextureData + Y * TGA->Width,ImageData + (TGA->Height - Y - 1) * TGA->Width,TGA->Width * 4);
+			FMemory::Memcpy(TextureData + Y * TGA->Width,ImageData + (TGA->Height - Y - 1) * TGA->Width * 4,TGA->Width * 4);
 		}
 
 		return true;
 	}
 
-	bool DecompressTGA_24bpp( const FTGAFileHeader* TGA, const int64 TGABufferLenght, uint32* TextureData )
+	bool DecompressTGA_24bpp( const FTGAFileHeader* TGA, const int64 TGABufferLength, uint32* TextureData )
 	{
-		const uint8* const IdData = (uint8*)TGA + sizeof(FTGAFileHeader);
-		const uint8* const ColorMap = IdData + TGA->IdFieldLength;
-		const uint8* const ImageData = (uint8*) (ColorMap + (TGA->ColorMapEntrySize + 4) / 8 * (int64)TGA->ColorMapLength);
-		uint8 Pixel[4];
+		const uint8* ImageData = GetImageData(TGA,TGABufferLength);
+		if ( ImageData == nullptr )
+			return false;
 
-		if ( (uint8*)TGA + TGABufferLenght < ImageData + (int64)TGA->Width * TGA->Height * 3 )
+		if ( TGABufferLength < ((uint8*)ImageData - (uint8*)TGA) + (int64)TGA->Width * TGA->Height * 3 )
 		{
 			return false;
 		}
@@ -221,27 +254,22 @@ namespace TgaImageSupportImpl
 		{
 			for(int64 X = 0; X < TGA->Width; X++)
 			{
-				Pixel[0] = *(( ImageData+( TGA->Height-Y-1 )*TGA->Width*3 )+X*3+0);
-				Pixel[1] = *(( ImageData+( TGA->Height-Y-1 )*TGA->Width*3 )+X*3+1);
-				Pixel[2] = *(( ImageData+( TGA->Height-Y-1 )*TGA->Width*3 )+X*3+2);
-				Pixel[3] = 255;
-				*((TextureData+Y*TGA->Width)+X) = *(uint32*)&Pixel;
+				const uint8 * BGR = ImageData+( TGA->Height-Y-1 )*TGA->Width*3 +X*3;
+
+				TextureData[Y*TGA->Width+X] = ConvertBGRToBGRA8(BGR);
 			}
 		}
 
 		return true;
 	}
 
-	bool DecompressTGA_16bpp( const FTGAFileHeader* TGA, const int64 TGABufferLenght, uint32* TextureData )
+	bool DecompressTGA_16bpp( const FTGAFileHeader* TGA, const int64 TGABufferLength, uint32* TextureData )
 	{
-		//todo: fix : this kind of code is prone to overflow bugs :
-		const uint8* const IdData = (uint8*)TGA + sizeof(FTGAFileHeader);
-		const uint8* const ColorMap = IdData + TGA->IdFieldLength;
-		const uint16* ImageData = (uint16*) (ColorMap + (TGA->ColorMapEntrySize + 4) / 8 * (int64)TGA->ColorMapLength);
-		uint16 FilePixel = 0;
-		uint32 TexturePixel = 0;
+		const uint8* ImageData = GetImageData(TGA,TGABufferLength);
+		if ( ImageData == nullptr )
+			return false;
 
-		if ( (uint16*)((uint8*)TGA + TGABufferLenght) < ImageData + (int64)TGA->Height * TGA->Width )
+		if ( TGABufferLength < ((uint8*)ImageData - (uint8*)TGA) + (int64)TGA->Height * TGA->Width * 2 )
 		{
 			return false;
 		}
@@ -250,27 +278,21 @@ namespace TgaImageSupportImpl
 		{					
 			for (int64 X = 0; X < TGA->Width; X++)
 			{
-				FilePixel = *ImageData++;
-				// Convert file format A1R5G5B5 into pixel format B8G8R8A8
-				TexturePixel = (FilePixel & 0x001F) << 3;
-				TexturePixel |= (FilePixel & 0x03E0) << 6;
-				TexturePixel |= (FilePixel & 0x7C00) << 9;
-				TexturePixel |= (FilePixel & 0x8000) << 16;
-				// Store.
-				*((TextureData + Y*TGA->Width) + X) = TexturePixel;						
+				*((TextureData + Y*TGA->Width) + X) = ConvertA1R5G5B5ToBGRA8(ImageData);
+				ImageData += 2;					
 			}
 		}
 
 		return true;
 	}
 
-	bool DecompressTGA_8bpp( const FTGAFileHeader* TGA, const int64 TGABufferLenght, uint8* TextureData )
+	bool DecompressTGA_8bpp( const FTGAFileHeader* TGA, const int64 TGABufferLength, uint8* TextureData )
 	{
-		const uint8* const IdData = (uint8*)TGA + sizeof(FTGAFileHeader);
-		const uint8* const  ColorMap = IdData + TGA->IdFieldLength;
-		const uint8* const  ImageData = (uint8*) (ColorMap + (TGA->ColorMapEntrySize + 4) / 8 * (int64)TGA->ColorMapLength);
-
-		if ( (uint8*)TGA + TGABufferLenght < ImageData + (int64)TGA->Width * TGA->Height )
+		const uint8* ImageData = GetImageData(TGA,TGABufferLength);
+		if ( ImageData == nullptr )
+			return false;
+			
+		if ( TGABufferLength < ((uint8*)ImageData - (uint8*)TGA) + (int64)TGA->Width * TGA->Height )
 		{
 			return false;
 		}
@@ -287,26 +309,26 @@ namespace TgaImageSupportImpl
 	}
 }
 
-bool DecompressTGA_helper( const FTGAFileHeader* TgaHeader, const int64 TGABufferLenght, uint32*& TextureData, const int64 TextureDataSize )
+bool DecompressTGA_helper( const FTGAFileHeader* TgaHeader, const int64 TGABufferLength, uint32* TextureData, const int64 TextureDataSize )
 {
 	bool bSuccess = false;
 	if ( TgaHeader->ImageTypeCode == 10 ) // 10 = RLE compressed 
 	{
-		check( TextureDataSize == TgaHeader->Width * TgaHeader->Height * 4 );
+		check( TextureDataSize == (int64) TgaHeader->Width * TgaHeader->Height * 4 );
 
 		// RLE compression: CHUNKS: 1 -byte header, high bit 0 = raw, 1 = compressed
 		// bits 0-6 are a 7-bit count; count+1 = number of raw pixels following, or rle pixels to be expanded. 
 		if(TgaHeader->BitsPerPixel == 32)
 		{
-			bSuccess = TgaImageSupportImpl::DecompressTGA_RLE_32bpp(TgaHeader, TGABufferLenght, TextureData);
+			bSuccess = TgaImageSupportImpl::DecompressTGA_RLE_32bpp(TgaHeader, TGABufferLength, TextureData);
 		}
 		else if( TgaHeader->BitsPerPixel == 24 )
 		{
-			bSuccess = TgaImageSupportImpl::DecompressTGA_RLE_24bpp(TgaHeader, TGABufferLenght, TextureData);
+			bSuccess = TgaImageSupportImpl::DecompressTGA_RLE_24bpp(TgaHeader, TGABufferLength, TextureData);
 		}
 		else if( TgaHeader->BitsPerPixel == 16 )
 		{
-			bSuccess = TgaImageSupportImpl::DecompressTGA_RLE_16bpp(TgaHeader, TGABufferLenght, TextureData);
+			bSuccess = TgaImageSupportImpl::DecompressTGA_RLE_16bpp(TgaHeader, TGABufferLength, TextureData);
 		}
 		else
 		{
@@ -320,15 +342,15 @@ bool DecompressTGA_helper( const FTGAFileHeader* TgaHeader, const int64 TGABuffe
 
 		if(TgaHeader->BitsPerPixel == 32)
 		{
-			bSuccess = TgaImageSupportImpl::DecompressTGA_32bpp(TgaHeader, TGABufferLenght, TextureData);
+			bSuccess = TgaImageSupportImpl::DecompressTGA_32bpp(TgaHeader, TGABufferLength, TextureData);
 		}
 		else if(TgaHeader->BitsPerPixel == 16)
 		{
-			bSuccess = TgaImageSupportImpl::DecompressTGA_16bpp(TgaHeader, TGABufferLenght, TextureData);
+			bSuccess = TgaImageSupportImpl::DecompressTGA_16bpp(TgaHeader, TGABufferLength, TextureData);
 		}
 		else if(TgaHeader->BitsPerPixel == 24)
 		{
-			bSuccess = TgaImageSupportImpl::DecompressTGA_24bpp(TgaHeader, TGABufferLenght, TextureData);
+			bSuccess = TgaImageSupportImpl::DecompressTGA_24bpp(TgaHeader, TGABufferLength, TextureData);
 		}
 		else
 		{
@@ -339,16 +361,16 @@ bool DecompressTGA_helper( const FTGAFileHeader* TgaHeader, const int64 TGABuffe
 	// Support for alpha stored as pseudo-color 8-bit TgaHeader
 	else if(TgaHeader->ColorMapType == 1 && TgaHeader->ImageTypeCode == 1 && TgaHeader->BitsPerPixel == 8)
 	{
-		check( TextureDataSize == TgaHeader->Width * TgaHeader->Height * 1 );
+		check( TextureDataSize == (int64) TgaHeader->Width * TgaHeader->Height * 1 );
 
-		bSuccess = TgaImageSupportImpl::DecompressTGA_8bpp(TgaHeader, TGABufferLenght, (uint8*)TextureData);
+		bSuccess = TgaImageSupportImpl::DecompressTGA_8bpp(TgaHeader, TGABufferLength, (uint8*)TextureData);
 	}
 	// standard grayscale
 	else if(TgaHeader->ColorMapType == 0 && TgaHeader->ImageTypeCode == 3 && TgaHeader->BitsPerPixel == 8)
 	{
-		check( TextureDataSize == TgaHeader->Width * TgaHeader->Height * 1 );
+		check( TextureDataSize == (int64) TgaHeader->Width * TgaHeader->Height * 1 );
 
-		bSuccess = TgaImageSupportImpl::DecompressTGA_8bpp(TgaHeader, TGABufferLenght, (uint8*)TextureData);
+		bSuccess = TgaImageSupportImpl::DecompressTGA_8bpp(TgaHeader, TGABufferLength, (uint8*)TextureData);
 	}
 	else
 	{
@@ -376,23 +398,37 @@ bool DecompressTGA_helper( const FTGAFileHeader* TgaHeader, const int64 TGABuffe
 		
 		check( TextureDataSize == NumBlocksX * NumBlocksY * BlockBytes );
 
-		uint8* MipData = (uint8*)TextureData;
+		const uint8* MipData = (uint8*)TextureData;
+		uint8 * FlippedDataPtr = &FlippedData[0];
 
-		for( int64 Y = 0; Y < NumBlocksY;Y++ )
+		if ( BlockBytes == 1 )
 		{
-			for( int64 X  = 0; X < NumBlocksX; X++ )
+			for( int64 Y = 0; Y < NumBlocksY;Y++ )
 			{
-				int64 DestX = bFlipX ? (NumBlocksX - X - 1) : X;
-				int64 DestY = bFlipY ? (NumBlocksY - Y - 1) : Y;
-				FMemory::Memcpy(
-					&FlippedData[(DestX + DestY * NumBlocksX) * BlockBytes],
-					&MipData[(X + Y * NumBlocksX) * BlockBytes],
-					BlockBytes
-				);
+				for( int64 X  = 0; X < NumBlocksX; X++ )
+				{
+					int64 DestX = bFlipX ? (NumBlocksX - X - 1) : X;
+					int64 DestY = bFlipY ? (NumBlocksY - Y - 1) : Y;
+					
+					FlippedDataPtr[DestX + DestY * NumBlocksX] = MipData[X + Y * NumBlocksX];
+				}
+			}
+		}
+		else
+		{
+			for( int64 Y = 0; Y < NumBlocksY;Y++ )
+			{
+				for( int64 X  = 0; X < NumBlocksX; X++ )
+				{
+					int64 DestX = bFlipX ? (NumBlocksX - X - 1) : X;
+					int64 DestY = bFlipY ? (NumBlocksY - Y - 1) : Y;
+					
+					((uint32 *)FlippedDataPtr)[DestX + DestY * NumBlocksX] = ((uint32 *)MipData)[X + Y * NumBlocksX];
+				}
 			}
 		}
 
-		FMemory::Memcpy( MipData, FlippedData.GetData(), FlippedData.Num() );
+		FMemory::Memcpy( TextureData, FlippedDataPtr, FlippedData.Num() );
 	}
 
 	return true;
diff --git a/Engine/Source/Runtime/ImageWrapper/Public/TgaImageSupport.h b/Engine/Source/Runtime/ImageWrapper/Public/TgaImageSupport.h
--- a/Engine/Source/Runtime/ImageWrapper/Public/TgaImageSupport.h
+++ b/Engine/Source/Runtime/ImageWrapper/Public/TgaImageSupport.h
@@ -37,4 +37,4 @@ struct FTGAFileHeader
  * This helper allows to decompress TGA data in a pre-allocated memory block.
  * The pixel format is necessarily PF_A8R8G8B8.
  */
-IMAGEWRAPPER_API bool DecompressTGA_helper(	const FTGAFileHeader* TGA, const int64 TGABufferLength, uint32*& TextureData, const int64 TextureDataSize );
+IMAGEWRAPPER_API bool DecompressTGA_helper(	const FTGAFileHeader* TGA, const int64 TGABufferLength, uint32* TextureData, const int64 TextureDataSize );
