TiffImageWrapper : fix missing Reset function so state carried over in multiple uses of Wrapper

#rnx

diff --git a/Engine/Source/Runtime/ImageWrapper/Private/Formats/TiffImageWrapper.cpp b/Engine/Source/Runtime/ImageWrapper/Private/Formats/TiffImageWrapper.cpp
--- a/Engine/Source/Runtime/ImageWrapper/Private/Formats/TiffImageWrapper.cpp
+++ b/Engine/Source/Runtime/ImageWrapper/Private/Formats/TiffImageWrapper.cpp
@@ -446,6 +446,12 @@ namespace UE::ImageWrapper::Private
 
 	void FTiffImageWrapper::Uncompress(const ERGBFormat InFormat, int32 InBitDepth)
 	{
+		if ( Tiff == nullptr )
+		{
+			SetError(TEXT("Tiff invalid."));
+			return;
+		}
+
 		if (InFormat == Format && InBitDepth == BitDepth)
 		{
 			// Read using RGBA
@@ -524,6 +530,8 @@ namespace UE::ImageWrapper::Private
 		{
 			SetError(TEXT("Unsupported requested format for the input image. Can't uncompress the tiff image."));
 		}
+
+		ReleaseTiffImage();
 	}
 
 	bool FTiffImageWrapper::SetCompressed(const void* InCompressedData, int64 InCompressedSize)
@@ -687,12 +695,13 @@ namespace UE::ImageWrapper::Private
 					break;
 				}
 
-
 				if (Format == ERGBFormat::Invalid)
 				{
 					SetError(TEXT("Unsupported Tiff content."));
 					return false;
 				}
+				
+				// note: the "Tiff" object is retained until the Uncompress call
 			}
 			else
 			{
diff --git a/Engine/Source/Runtime/ImageWrapper/Private/Formats/TiffImageWrapper.h b/Engine/Source/Runtime/ImageWrapper/Private/Formats/TiffImageWrapper.h
--- a/Engine/Source/Runtime/ImageWrapper/Private/Formats/TiffImageWrapper.h
+++ b/Engine/Source/Runtime/ImageWrapper/Private/Formats/TiffImageWrapper.h
@@ -64,6 +64,20 @@ namespace UE::ImageWrapper::Private
 		uint16 SampleFormat = 0;
 
 		friend struct FTIFFReadMemoryFile;
+		
+		virtual void Reset() override
+		{
+			FImageWrapperBase::Reset();
+
+			ReleaseTiffImage();
+			
+			CurrentPosition = 0;
+			Photometric = 0;
+			Compression = 0;
+			BitsPerSample = 0;
+			SamplesPerPixel = 0;
+			SampleFormat = 0;
+		}
 	};
 
 
