ImageCoreUtils::IsImageImportPossible can be checked before image import instead of after

#rnx

diff --git a/Engine/Plugins/Interchange/Runtime/Source/Import/Private/Texture/InterchangeTexturePayloadData.cpp b/Engine/Plugins/Interchange/Runtime/Source/Import/Private/Texture/InterchangeTexturePayloadData.cpp
--- a/Engine/Plugins/Interchange/Runtime/Source/Import/Private/Texture/InterchangeTexturePayloadData.cpp
+++ b/Engine/Plugins/Interchange/Runtime/Source/Import/Private/Texture/InterchangeTexturePayloadData.cpp
@@ -7,6 +7,7 @@
 #include "HAL/IConsoleManager.h"
 #include "Misc/MessageDialog.h"
 #include "RHI.h"
+#include "ImageCoreUtils.h"
 
 bool UE::Interchange::FImportImageHelper::IsImportResolutionValid(int64 Width, int64 Height, bool bAllowNonPowerOfTwo, FText* OutErrorMessage)
 {
@@ -35,16 +36,9 @@ bool UE::Interchange::FImportImageHelper::IsImportResolutionValid(int64 Width, i
 		return false;
 	}
 
-	bool bTextureTooLargeOrInvalidResolution = false;
-
 	// Dimensions must fit in signed int32
 	//  could be negative here if it was over 2G and int32 was used earlier
-	if (Width < 0 || Height < 0 || Width > MAX_int32 || Height > MAX_int32)
-	{
-		bTextureTooLargeOrInvalidResolution = true;
-	}
-
-	if (bTextureTooLargeOrInvalidResolution)
+	if ( ! FImageCoreUtils::IsImageImportPossible(Width,Height) )
 	{
 		if (OutErrorMessage)
 		{
diff --git a/Engine/Source/Editor/UnrealEd/Private/Factories/EditorFactories.cpp b/Engine/Source/Editor/UnrealEd/Private/Factories/EditorFactories.cpp
--- a/Engine/Source/Editor/UnrealEd/Private/Factories/EditorFactories.cpp
+++ b/Engine/Source/Editor/UnrealEd/Private/Factories/EditorFactories.cpp
@@ -4513,7 +4513,7 @@ bool UTextureFactory::IsImportResolutionValid(int64 Width, int64 Height, bool bA
 
 	// Dimensions must fit in signed int32
 	//  could be negative here if it was over 2G and int32 was used earlier
-	if (Width < 0 || Height < 0 || Width > MAX_int32 || Height > MAX_int32)
+	if ( ! FImageCoreUtils::IsImageImportPossible(Width,Height) )
 	{
 		Warn->Log(ELogVerbosity::Error, NSLOCTEXT("UnrealEd", "Warning_TextureSizeTooLargeOrInvalid", "Texture is has an invalid resolution.").ToString());
 
diff --git a/Engine/Source/Runtime/ImageCore/ImageCore.Build.cs b/Engine/Source/Runtime/ImageCore/ImageCore.Build.cs
--- a/Engine/Source/Runtime/ImageCore/ImageCore.Build.cs
+++ b/Engine/Source/Runtime/ImageCore/ImageCore.Build.cs
@@ -8,6 +8,7 @@ public class ImageCore : ModuleRules
 	{
 		// include only, no link :
 		// for TextureDefines.h :
+		// PublicIncludePathModuleNames
 		PrivateIncludePathModuleNames.AddRange(new string[]
 		{
 			"Engine",
diff --git a/Engine/Source/Runtime/ImageCore/Private/ImageCoreUtils.cpp b/Engine/Source/Runtime/ImageCore/Private/ImageCoreUtils.cpp
--- a/Engine/Source/Runtime/ImageCore/Private/ImageCoreUtils.cpp
+++ b/Engine/Source/Runtime/ImageCore/Private/ImageCoreUtils.cpp
@@ -4,6 +4,34 @@
 
 DEFINE_LOG_CATEGORY_STATIC(LogImageCoreUtils, Log, All);
 
+IMAGECORE_API bool FImageCoreUtils::IsImageImportPossible(int64 Width,int64 Height)
+{
+	if ( Width <= 0 || Height <= 0 )
+	{
+		return false;
+	}
+
+	// each dimension must fit in int32
+	if ( Width > INT32_MAX || Height > INT32_MAX )
+	{
+		return false;
+	}
+
+	int64 PixelCount = Width * Height; // mutliply is safe because of INT32_MAX check
+
+	// for non-VT limitation is 16384 for W/H , and 2 GB for output surface bytes
+	// for VT the limitation is that the built VT must fit in 4 GB
+
+	if ( PixelCount > UINT32_MAX )
+	{
+		return false;
+	}
+
+	// still may fail to import due to stricter constraints later
+
+	return true;
+}
+
 IMAGECORE_API int32 FImageCoreUtils::GetMipCountFromDimensions(int32 InSizeX, int32 InSizeY, int32 InVolumeZ, bool bInIsVolume)
 {
 	uint32 MaxMipDimension = FMath::Max3(InSizeX, InSizeY, bInIsVolume ? InVolumeZ : 1);
diff --git a/Engine/Source/Runtime/ImageCore/Public/ImageCoreUtils.h b/Engine/Source/Runtime/ImageCore/Public/ImageCoreUtils.h
--- a/Engine/Source/Runtime/ImageCore/Public/ImageCoreUtils.h
+++ b/Engine/Source/Runtime/ImageCore/Public/ImageCoreUtils.h
@@ -20,6 +20,13 @@ ImageCoreUtils.h: Image utility functions.
 
 namespace FImageCoreUtils
 {
+	// IsImageImportPossible returns false if the image dimensions are not supported to import
+	//	if this function returns true, image import might still fail due to further constraints
+	//	(this constraint is weaker than the one actually required)
+	//  (see IsImportResolutionValid, which is also still weaker than the final constraint)
+	// IsImageImportPossible does not check current Engine config or platform limits, only hard-coded absolute limits
+	IMAGECORE_API bool IsImageImportPossible(int64 Width,int64 Height);
+
 	// Returns the number of mips that constitute a full mip chain for the given top level mip size.
 	// InVolumeZ is ignored unless bInIsVolume is true
 	IMAGECORE_API int32 GetMipCountFromDimensions(int32 InSizeX, int32 InSizeY, int32 InVolumeZ, bool bInIsVolume);
diff --git a/Engine/Source/Runtime/ImageWrapper/ImageWrapper.Build.cs b/Engine/Source/Runtime/ImageWrapper/ImageWrapper.Build.cs
--- a/Engine/Source/Runtime/ImageWrapper/ImageWrapper.Build.cs
+++ b/Engine/Source/Runtime/ImageWrapper/ImageWrapper.Build.cs
@@ -13,6 +13,13 @@ public class ImageWrapper : ModuleRules
 			}
 		);
 
+		// include only, no link :
+		// for TextureDefines.h :
+		PrivateIncludePathModuleNames.AddRange(new string[]
+		{
+			"Engine",
+		});
+
 		PublicDefinitions.Add("WITH_UNREALPNG=1");
 		PublicDefinitions.Add("WITH_UNREALJPEG=1");
 
diff --git a/Engine/Source/Runtime/ImageWrapper/Private/Formats/BmpImageWrapper.cpp b/Engine/Source/Runtime/ImageWrapper/Private/Formats/BmpImageWrapper.cpp
--- a/Engine/Source/Runtime/ImageWrapper/Private/Formats/BmpImageWrapper.cpp
+++ b/Engine/Source/Runtime/ImageWrapper/Private/Formats/BmpImageWrapper.cpp
@@ -4,6 +4,7 @@
 #include "ImageWrapperPrivate.h"
 
 #include "BmpImageSupport.h"
+#include "ImageCoreUtils.h"
 
 static inline bool BmpDimensionIsValid(int32 Dim)
 {
@@ -466,6 +467,13 @@ bool FBmpImageWrapper::SetCompressed(const void* InCompressedData, int64 InCompr
 	if ( ! bResult )
 	{
 		CompressedData.Reset();
+		return false;
+	}
+	
+	if ( ! FImageCoreUtils::IsImageImportPossible(Width,Height) )
+	{
+		SetError(TEXT("Image dimensions are not possible to import"));
+		return false;
 	}
 
 	return bResult;
diff --git a/Engine/Source/Runtime/ImageWrapper/Private/Formats/DdsImageWrapper.cpp b/Engine/Source/Runtime/ImageWrapper/Private/Formats/DdsImageWrapper.cpp
--- a/Engine/Source/Runtime/ImageWrapper/Private/Formats/DdsImageWrapper.cpp
+++ b/Engine/Source/Runtime/ImageWrapper/Private/Formats/DdsImageWrapper.cpp
@@ -2,6 +2,7 @@
 
 #include "Formats/DdsImageWrapper.h"
 #include "ImageWrapperPrivate.h"
+#include "ImageCoreUtils.h"
 
 void FDdsImageWrapper::Reset()
 {
@@ -158,6 +159,12 @@ bool FDdsImageWrapper::SetCompressed(const void* InCompressedData, int64 InCompr
 	Width = DDS->Width;
 	Height = DDS->Height;
 
+	if ( ! FImageCoreUtils::IsImageImportPossible(Width,Height) )
+	{
+		SetError(TEXT("Image dimensions are not possible to import"));
+		return false;
+	}
+
 	return true;
 }
 
diff --git a/Engine/Source/Runtime/ImageWrapper/Private/Formats/ExrImageWrapper.cpp b/Engine/Source/Runtime/ImageWrapper/Private/Formats/ExrImageWrapper.cpp
--- a/Engine/Source/Runtime/ImageWrapper/Private/Formats/ExrImageWrapper.cpp
+++ b/Engine/Source/Runtime/ImageWrapper/Private/Formats/ExrImageWrapper.cpp
@@ -8,6 +8,7 @@
 #include "Containers/StringConv.h"
 #include "HAL/PlatformTime.h"
 #include "Math/Float16.h"
+#include "ImageCoreUtils.h"
 
 #if WITH_UNREALEXR
 
@@ -91,9 +92,14 @@ public:
 		, Size(InSize)
 		, Pos(0)
 	{
+		// InSize is signed but "Size" member is unsigned
 		if (InSize < 0)
 		{
-			UE_LOG(LogImageWrapper, Fatal, TEXT("Negative size passed to EXR parser, can not continue."));
+			Size = 0;
+			// MSVC runtime_error vtable size changes based on module bEnableExceptions setting
+			//throw std::runtime_error("FMemFileIn: Negative size passed to EXR parser");
+			// MSVC std::exception has a char * constructor, but that is not portable, do not use
+			throw std::exception();
 		}
 	}
 
@@ -108,28 +114,30 @@ public:
 	//------------------------------------------------------
 
 	// InN must be 32bit to match the abstract interface.
-	virtual bool read (char c[/*n*/], int32 InN)
+	virtual bool read (char Out[/*n*/], int32 Count)
 	{
-		int64 SrcN = InN;
+		// return false if EOF is hit
+		// OpenEXR mostly ignores this return value, you must throw to get error handling
 
-		FGuardedInt64 NextPosition = FGuardedInt64(Pos) + SrcN;
-		if (SrcN < 0 ||
+		FGuardedInt64 NextPosition = FGuardedInt64(Pos) + Count;
+		if (Count < 0 ||
 			NextPosition.InvalidOrGreaterThan(Size))
 		{
-			// This is supposed to throw an exception per the spec, however we can't use
-			// an actual exception type right now due to linker issues during LTCG.
-			UE_LOG(LogImageWrapper, Error, TEXT("Invalid read requested in EXR"));
+			//throw std::runtime_error("FMemFileIn: Exr read out of bounds");
 			throw std::exception();
-			return false;
 		}
 
-		for (int64 i = 0; i < SrcN; ++i)
+		memcpy(Out,Data+Pos,Count);
+		Pos += Count; // == NextPosition
+
+		if ( Pos == Size )
 		{
-			c[i] = Data[Pos];
-			++Pos;
+			return false;
+		}
+		else
+		{
+			return true;
 		}
-
-		return Pos >= Size;
 	}
 
 	//--------------------------------------------------------
@@ -299,6 +307,12 @@ bool FExrImageWrapper::SetCompressed(const void* InCompressedData, int64 InCompr
 		return false;
 	}
 
+	if ( ! FImageCoreUtils::IsImageImportPossible(Width,Height) )
+	{
+		SetError(TEXT("Image dimensions are not possible to import"));
+		return false;
+	}
+
 	return true;
 }
 
@@ -553,9 +567,12 @@ void FExrImageWrapper::Uncompress(const ERGBFormat InFormat, const int32 InBitDe
 			}
 			double DefaultValue = bIsAlphaChannel ? 1.0 : 0.0;
 
-			// @todo Oodle: doing a subtract on this pointer for the data window looks dangerous
-			char* ChannelBase = (char*)(ChannelData[c].GetData() -
-				(ImfDataWindow.min.x + ImfDataWindow.min.y * Width) * ((int64)BytesPerChannelPixel));
+			int64 DataWindowOffset = (ImfDataWindow.min.x + ImfDataWindow.min.y * Width) * ((int64)BytesPerChannelPixel);
+			check( (ImfDataWindow.max.x - ImfDataWindow.min.x + 1) == Width );
+			check( (ImfDataWindow.max.y - ImfDataWindow.min.y + 1) == Height );
+
+			// OpenExr does this offset with the pointers cast to intptr_t ; see ImfFrameBuffer.cpp
+			char* ChannelBase = (char*)((intptr_t)ChannelData[c].GetData() - DataWindowOffset);
 
 			ImfFrameBuffer.insert(ChannelNames[c], Imf::Slice(ImfPixelType, ChannelBase, BytesPerChannelPixel, (size_t)BytesPerChannelPixel * Width, 1, 1, DefaultValue));
 		}
diff --git a/Engine/Source/Runtime/ImageWrapper/Private/Formats/HdrImageWrapper.cpp b/Engine/Source/Runtime/ImageWrapper/Private/Formats/HdrImageWrapper.cpp
--- a/Engine/Source/Runtime/ImageWrapper/Private/Formats/HdrImageWrapper.cpp
+++ b/Engine/Source/Runtime/ImageWrapper/Private/Formats/HdrImageWrapper.cpp
@@ -2,6 +2,7 @@
 
 #include "Formats/HdrImageWrapper.h"
 #include "ImageWrapperPrivate.h"
+#include "ImageCoreUtils.h"
 
 #define LOCTEXT_NAMESPACE "HdrImageWrapper"
 
@@ -116,6 +117,13 @@ bool FHdrImageWrapper::SetCompressedFromView(TArrayView64<const uint8> Data)
 	Width = ImageWidth;
 	Height = ImageHeight;
 	RGBDataStart = FileDataPtr;
+	
+	if ( ! FImageCoreUtils::IsImageImportPossible(Width,Height) )
+	{
+		SetAndLogError(LOCTEXT("ImpossibleImport","Image dimensions are not possible to import"));
+		return false;
+	}
+
 	return true;
 }
 
diff --git a/Engine/Source/Runtime/ImageWrapper/Private/Formats/IcnsImageWrapper.cpp b/Engine/Source/Runtime/ImageWrapper/Private/Formats/IcnsImageWrapper.cpp
--- a/Engine/Source/Runtime/ImageWrapper/Private/Formats/IcnsImageWrapper.cpp
+++ b/Engine/Source/Runtime/ImageWrapper/Private/Formats/IcnsImageWrapper.cpp
@@ -44,8 +44,19 @@ bool FIcnsImageWrapper::SetCompressed(const void* InCompressedData, int64 InComp
 			[Image release];
 		}
 	}
+	
+	if ( Width == 0 )
+	{
+		return false;
+	}
+
+	if ( ! FImageCoreUtils::IsImageImportPossible(Width,Height) )
+	{
+		SetError(TEXT("Image dimensions are not possible to import"));
+		return false;
+	}
 
-	return ( Width > 0 );
+	return true;
 #else
 	return false;
 #endif
diff --git a/Engine/Source/Runtime/ImageWrapper/Private/Formats/PngImageWrapper.cpp b/Engine/Source/Runtime/ImageWrapper/Private/Formats/PngImageWrapper.cpp
--- a/Engine/Source/Runtime/ImageWrapper/Private/Formats/PngImageWrapper.cpp
+++ b/Engine/Source/Runtime/ImageWrapper/Private/Formats/PngImageWrapper.cpp
@@ -5,7 +5,7 @@
 
 #include "Math/GuardedInt.h"
 #include "Misc/ScopeLock.h"
-
+#include "ImageCoreUtils.h"
 
 #if WITH_UNREALPNG
 
@@ -15,9 +15,10 @@
 	#pragma warning(disable:4611)
 #endif
 
+#if PLATFORM_ANDROID
 /** Only allow one thread to use libpng at a time (it's not thread safe) */
-FCriticalSection GPNGSection;
-
+static FCriticalSection GPNGSection;
+#endif
 
 /* Local helper classes
  *****************************************************************************/
@@ -301,6 +302,8 @@ bool FPngImageWrapper::SetCompressed(const void* InCompressedData, int64 InCompr
 			{
 				check(CompressedData.Num());
 				UncompressPNGData(Format, 8);
+
+				// after UncompressPNGData , BitDepth is now changed to 8
 			}
 		}
 		else if (ColorType == PNG_COLOR_TYPE_RGB)
@@ -312,6 +315,13 @@ bool FPngImageWrapper::SetCompressed(const void* InCompressedData, int64 InCompr
 	if ( (Format == ERGBFormat::BGRA || Format == ERGBFormat::RGBA || Format == ERGBFormat::Gray) &&
 		(BitDepth == 8 || BitDepth == 16) )
 	{
+	
+		if ( ! FImageCoreUtils::IsImageImportPossible(Width,Height) )
+		{
+			SetError(TEXT("Image dimensions are not possible to import"));
+			return false;
+		}
+
 		return true;
 	}
 	else
@@ -534,8 +544,10 @@ bool FPngImageWrapper::LoadPNGHeader()
 	// Test whether the data this PNGLoader is pointing at is a PNG or not.
 	if (IsPNG())
 	{
+#if PLATFORM_ANDROID
 		// thread safety
 		FScopeLock PNGLock(&GPNGSection);
+#endif
 
 		png_structp png_ptr = png_create_read_struct_2(PNG_LIBPNG_VER_STRING, this, FPngImageWrapper::user_error_fn, FPngImageWrapper::user_warning_fn, NULL, FPngImageWrapper::user_malloc, FPngImageWrapper::user_free);
 		check(png_ptr);
diff --git a/Engine/Source/Runtime/ImageWrapper/Private/Formats/TgaImageWrapper.cpp b/Engine/Source/Runtime/ImageWrapper/Private/Formats/TgaImageWrapper.cpp
--- a/Engine/Source/Runtime/ImageWrapper/Private/Formats/TgaImageWrapper.cpp
+++ b/Engine/Source/Runtime/ImageWrapper/Private/Formats/TgaImageWrapper.cpp
@@ -4,6 +4,7 @@
 #include "ImageWrapperPrivate.h"
 
 #include "TgaImageSupport.h"
+#include "ImageCoreUtils.h"
 
 
 
@@ -111,7 +112,20 @@ bool FTgaImageWrapper::SetCompressed(const void* InCompressedData, int64 InCompr
 {
 	bool bResult = FImageWrapperBase::SetCompressed(InCompressedData, InCompressedSize);
 
-	return bResult && LoadTGAHeader();
+	if ( bResult && LoadTGAHeader() )
+	{
+		if ( ! FImageCoreUtils::IsImageImportPossible(Width,Height) )
+		{
+			SetError(TEXT("Image dimensions are not possible to import"));
+			return false;
+		}
+
+		return true;
+	}
+	else
+	{
+		return false;
+	}
 }
 
 void FTgaImageWrapper::Uncompress(const ERGBFormat InFormat, const int32 InBitDepth)
diff --git a/Engine/Source/Runtime/ImageWrapper/Private/Formats/TiffImageWrapper.cpp b/Engine/Source/Runtime/ImageWrapper/Private/Formats/TiffImageWrapper.cpp
--- a/Engine/Source/Runtime/ImageWrapper/Private/Formats/TiffImageWrapper.cpp
+++ b/Engine/Source/Runtime/ImageWrapper/Private/Formats/TiffImageWrapper.cpp
@@ -16,6 +16,7 @@
 #include "Misc/CoreStats.h"
 #include "Templates/IsSigned.h"
 #include "Templates/UnrealTypeTraits.h"
+#include "ImageCoreUtils.h"
 
 #include <type_traits>
 
@@ -536,9 +537,11 @@ namespace UE::ImageWrapper::Private
 
 	bool FTiffImageWrapper::SetCompressed(const void* InCompressedData, int64 InCompressedSize)
 	{
-		bool bResult = FImageWrapperBase::SetCompressed( InCompressedData, InCompressedSize );
+		if ( !  FImageWrapperBase::SetCompressed( InCompressedData, InCompressedSize ) )
+		{
+			return false;
+		}
 
-		if (bResult)
 		{
 			Tiff = TIFFClientOpen(""
 				, "r"
@@ -709,7 +712,13 @@ namespace UE::ImageWrapper::Private
 			}
 		}
 
-		return bResult;
+		if ( ! FImageCoreUtils::IsImageImportPossible(Width,Height) )
+		{
+			SetError(TEXT("Image dimensions are not possible to import"));
+			return false;
+		}
+
+		return true;
 
 	}
 
