IcnsImageWrapper fix Mac non-unity compile

#rnx

diff --git a/Engine/Source/Runtime/ImageWrapper/Private/Formats/IcnsImageWrapper.cpp b/Engine/Source/Runtime/ImageWrapper/Private/Formats/IcnsImageWrapper.cpp
--- a/Engine/Source/Runtime/ImageWrapper/Private/Formats/IcnsImageWrapper.cpp
+++ b/Engine/Source/Runtime/ImageWrapper/Private/Formats/IcnsImageWrapper.cpp
@@ -1,6 +1,9 @@
 // Copyright Epic Games, Inc. All Rights Reserved.
 
 #include "Formats/IcnsImageWrapper.h"
+#include "ImageWrapperPrivate.h"
+#include "ImageCoreUtils.h"
+
 
 
 /* FIcnsImageWrapper structors
