JpegImageWrapper : remove unnecessary serializing critsec; libjpeg-turbo is now thread safe
fix error state was being retained in tjhandles, make new ones for each encode
fix InterchangeJPGTranslator using deprecated ImageWrapper code style

#rnx

diff --git a/Engine/Plugins/Interchange/Runtime/Source/Import/Private/Texture/InterchangeJPGTranslator.cpp b/Engine/Plugins/Interchange/Runtime/Source/Import/Private/Texture/InterchangeJPGTranslator.cpp
--- a/Engine/Plugins/Interchange/Runtime/Source/Import/Private/Texture/InterchangeJPGTranslator.cpp
+++ b/Engine/Plugins/Interchange/Runtime/Source/Import/Private/Texture/InterchangeJPGTranslator.cpp
@@ -7,6 +7,7 @@
 #include "Engine/Texture2D.h"
 #include "IImageWrapper.h"
 #include "IImageWrapperModule.h"
+#include "ImageCoreUtils.h"
 #include "InterchangeImportLog.h"
 #include "InterchangeTextureNode.h"
 #include "Memory/SharedBuffer.h"
@@ -103,45 +104,19 @@ TOptional<UE::Interchange::FImportImage> UInterchangeJPGTranslator::GetTexturePa
 	}
 
 	// Select the texture's source format
-	ETextureSourceFormat TextureFormat = TSF_Invalid;
-	int32 BitDepth = JpegImageWrapper->GetBitDepth();
-	ERGBFormat Format = JpegImageWrapper->GetFormat();
-
-	if (Format == ERGBFormat::Gray)
-	{
-		if (BitDepth <= 8)
-		{
-			TextureFormat = TSF_G8;
-			Format = ERGBFormat::Gray;
-			BitDepth = 8;
-		}
-	}
-	else if (Format == ERGBFormat::RGBA)
-	{
-		if (BitDepth <= 8)
-		{
-			TextureFormat = TSF_BGRA8;
-			Format = ERGBFormat::BGRA;
-			BitDepth = 8;
-		}
-	}
-
-	if (TextureFormat == TSF_Invalid)
-	{
-		FTextureTranslatorUtilities::LogError(*this, NSLOCTEXT("InterchangeJPEGTranslator", "UnsupportedFormat", "JPEG file contains data in an unsupported format."));
-		return TOptional<UE::Interchange::FImportImage>();
-	}
+	ERawImageFormat::Type RawFormat = JpegImageWrapper->GetClosestRawImageFormat();
+	check( RawFormat != ERawImageFormat::Invalid );
+	ETextureSourceFormat TextureFormat = FImageCoreUtils::ConvertToTextureSourceFormat(RawFormat);
 
 	UE::Interchange::FImportImage PayloadData;
 
-
 	const bool bShouldAllocateRawDataBuffer = false;
 
 	PayloadData.Init2DWithParams(
 		JpegImageWrapper->GetWidth(),
 		JpegImageWrapper->GetHeight(),
 		TextureFormat,
-		BitDepth < 16,
+		JpegImageWrapper->GetSRGB(),
 		bShouldAllocateRawDataBuffer
 	);
 
@@ -151,7 +126,7 @@ TOptional<UE::Interchange::FImportImage> UInterchangeJPGTranslator::GetTexturePa
 		PayloadData.RawData = MakeUniqueBufferFromArray(MoveTemp(SourceDataBuffer));
 		PayloadData.RawDataCompressionFormat = ETextureSourceCompressionFormat::TSCF_JPEG;
 	}
-	else if (JpegImageWrapper->GetRaw(Format, BitDepth, RawData))
+	else if (JpegImageWrapper->GetRaw(RawData))
 	{
 		PayloadData.RawData = MakeUniqueBufferFromArray(MoveTemp(RawData));
 	}
diff --git a/Engine/Source/Runtime/ImageWrapper/Private/Formats/JpegImageWrapper.cpp b/Engine/Source/Runtime/ImageWrapper/Private/Formats/JpegImageWrapper.cpp
--- a/Engine/Source/Runtime/ImageWrapper/Private/Formats/JpegImageWrapper.cpp
+++ b/Engine/Source/Runtime/ImageWrapper/Private/Formats/JpegImageWrapper.cpp
@@ -5,6 +5,7 @@
 #include "Math/Color.h"
 #include "Misc/ScopeLock.h"
 #include "ImageWrapperPrivate.h"
+#include "ImageCoreUtils.h"
 
 #if WITH_UNREALJPEG
 
@@ -38,13 +39,13 @@
 	#pragma clang diagnostic pop
 #endif
 
-DEFINE_LOG_CATEGORY_STATIC(JPEGLog, Log, All);
-
 #if WITH_LIBJPEGTURBO
 namespace
 {
 	int ConvertTJpegPixelFormat(ERGBFormat InFormat)
 	{
+		// note: libjpeg-turbo currently does not actually read/write A
+		//	TJPF_BGRA is a synonym for TJPF_BGRX
 		switch (InFormat)
 		{
 			case ERGBFormat::BGRA:	return TJPF_BGRA;
@@ -56,38 +57,58 @@ namespace
 }
 #endif	// WITH_LIBJPEGTURBO
 
+
+#if WITH_LIBJPEGTURBO
+// libjpeg-turbo since version 2.0.5 is thread safe
+#define JPEG_NEEDS_CRITSEC 0
+#else
+// legacy libjpeg is not thread safe
+#define JPEG_NEEDS_CRITSEC 1
+#endif
+
+#if JPEG_NEEDS_CRITSEC
 // Only allow one thread to use JPEG decoder at a time (it's not thread safe)
-FCriticalSection GJPEGSection; 
+static FCriticalSection GJPEGSection; 
 
+#define JPEG_SCOPE_CRITSEC()	FScopeLock JPEGLock(&GJPEGSection)
+#else
+#define JPEG_SCOPE_CRITSEC()	do { } while(0)
+#endif
 
 /* FJpegImageWrapper structors
  *****************************************************************************/
 
-FJpegImageWrapper::FJpegImageWrapper(int32 InNumComponents)
+FJpegImageWrapper::FJpegImageWrapper(int32 InNumComponents /* = 4 */)
 	: FImageWrapperBase()
 	, NumComponents(InNumComponents)
 #if WITH_LIBJPEGTURBO
-	, Compressor(tjInitCompress())
-	, Decompressor(tjInitDecompress())
+	, Decompressor(0)
 #endif	// WITH_LIBJPEGTURBO
-{ }
+{ 
+	// NumComponents == 1 for GrayscaleJPEG
+}
 
 
-#if WITH_LIBJPEGTURBO
 FJpegImageWrapper::~FJpegImageWrapper()
 {
-	FScopeLock JPEGLock(&GJPEGSection);
+	Reset();
+}
 
-	if (Compressor)
-	{
-		tjDestroy(Compressor);
-	}
+void FJpegImageWrapper::Reset()
+{
+	FImageWrapperBase::Reset();
+
+#if WITH_LIBJPEGTURBO
 	if (Decompressor)
 	{
+		JPEG_SCOPE_CRITSEC();
+
 		tjDestroy(Decompressor);
+		Decompressor = 0;
 	}
+#endif
+
 }
-#endif	// WITH_LIBJPEGTURBO
 
 /* FImageWrapperBase interface
  *****************************************************************************/
@@ -107,9 +128,14 @@ bool FJpegImageWrapper::SetCompressed(const void* InCompressedData, int64 InComp
 
 	jpgd::jpeg_decoder decoder(&jpeg_memStream);
 	if (decoder.get_error_code() != jpgd::JPGD_SUCCESS)
+	{
 		return false;
+	}
 
-	bool bResult = FImageWrapperBase::SetCompressed(InCompressedData, InCompressedSize);
+	if ( ! FImageWrapperBase::SetCompressed(InCompressedData, InCompressedSize) )
+	{
+		return false;
+	}
 
 	// We don't support 16 bit jpegs
 	BitDepth = 8;
@@ -128,8 +154,14 @@ bool FJpegImageWrapper::SetCompressed(const void* InCompressedData, int64 InComp
 	default:
 		return false;
 	}
+	
+	if ( ! FImageCoreUtils::IsImageImportPossible(Width,Height) )
+	{
+		SetError(TEXT("Image dimensions are not possible to import"));
+		return false;
+	}
 
-	return bResult;
+	return true;
 #endif
 }
 
@@ -188,7 +220,7 @@ void FJpegImageWrapper::Compress(int32 Quality)
 #else
 	if (CompressedData.Num() == 0)
 	{
-		FScopeLock JPEGLock(&GJPEGSection);
+		JPEG_SCOPE_CRITSEC();
 		
 		check(RawData.Num());
 		check(Width > 0);
@@ -254,7 +286,7 @@ void FJpegImageWrapper::Uncompress(const ERGBFormat InFormat, int32 InBitDepth)
 		return;
 	}
 
-	FScopeLock JPEGLock(&GJPEGSection);
+	JPEG_SCOPE_CRITSEC();
 
 	check(CompressedData.Num());
 
@@ -283,9 +315,16 @@ void FJpegImageWrapper::Uncompress(const ERGBFormat InFormat, int32 InBitDepth)
 #if WITH_LIBJPEGTURBO
 bool FJpegImageWrapper::SetCompressedTurbo(const void* InCompressedData, int64 InCompressedSize)
 {
-	FScopeLock JPEGLock(&GJPEGSection);
+	// SetCompressed does Reset
+	if ( ! FImageWrapperBase::SetCompressed(InCompressedData, InCompressedSize) )
+	{
+		return false;
+	}
+	
+	JPEG_SCOPE_CRITSEC();
 
-	check(Decompressor);
+	check( Decompressor == 0 );
+	Decompressor = tjInitDecompress();
 
 	int ImageWidth;
 	int ImageHeight;
@@ -293,37 +332,47 @@ bool FJpegImageWrapper::SetCompressedTurbo(const void* InCompressedData, int64 I
 	int ColorSpace;
 	if (tjDecompressHeader3(Decompressor, reinterpret_cast<const uint8*>(InCompressedData), InCompressedSize, &ImageWidth, &ImageHeight, &SubSampling, &ColorSpace) != 0)
 	{
+		SetError(TEXT("tjDecompressHeader3 failed"));
 		return false;
 	}
 
-	const bool bResult = FImageWrapperBase::SetCompressed(InCompressedData, InCompressedSize);
-
 	// set after call to base SetCompressed as it will reset members
 	Width = ImageWidth;
 	Height = ImageHeight;
 	BitDepth = 8; // We don't support 16 bit jpegs
-	Format = SubSampling == TJSAMP_GRAY ? ERGBFormat::Gray : ERGBFormat::RGBA;
 
-	return bResult;
+	// if NumComponents == 1 (for GrayscaleJPEG format), force ERGBFormat::Gray ?
+	Format = ( SubSampling == TJSAMP_GRAY ) ? ERGBFormat::Gray : ERGBFormat::BGRA;
+	
+	if ( ! FImageCoreUtils::IsImageImportPossible(Width,Height) )
+	{
+		SetError(TEXT("Image dimensions are not possible to import"));
+		return false;
+	}
+
+	// Decompressor is retained until Uncompress
+
+	return true;
 }
 
 void FJpegImageWrapper::CompressTurbo(int32 Quality)
 {
 	if (CompressedData.Num() == 0)
 	{
-		FScopeLock JPEGLock(&GJPEGSection);
+		JPEG_SCOPE_CRITSEC();
 
 		// Quality mapping should have already been done
 		check( Quality >= JPEG_QUALITY_MIN && Quality <= 100 );
 
-		check(Compressor);
 		check(RawData.Num());
 		check(Width > 0);
 		check(Height > 0);
 		check(BitDepth == 8);
 
 		const int PixelFormat = ConvertTJpegPixelFormat(Format);
-		const int Subsampling = TJSAMP_420;
+
+		// NumComponents == 1 for GrayscaleJPEG
+		const int Subsampling = (NumComponents == 1 || Format == ERGBFormat::Gray) ? TJSAMP_GRAY : TJSAMP_420;
 		const int Flags = TJFLAG_NOREALLOC | TJFLAG_FASTDCT;
 
 		unsigned long OutBufferSize = tjBufSize(Width, Height, Subsampling);
@@ -331,11 +380,24 @@ void FJpegImageWrapper::CompressTurbo(int32 Quality)
 		unsigned char* OutBuffer = CompressedData.GetData();
 
 		int BytesPerRow = GetBytesPerRow();
+		
+		tjhandle Compressor = tjInitCompress();
 
 		const bool bSuccess = tjCompress2(Compressor, RawData.GetData(), Width, BytesPerRow, Height, PixelFormat, &OutBuffer, &OutBufferSize, Subsampling, Quality, Flags) == 0;
 		check(bSuccess);
+		
+		tjDestroy(Compressor);
 
-		CompressedData.SetNum(OutBufferSize);
+		if ( ! bSuccess )
+		{
+			CompressedData.Empty();
+			SetError(TEXT("tjCompress2 failed"));
+		}
+		else
+		{
+			check( CompressedData.GetData() == OutBuffer ); // TJFLAG_NOREALLOC so OutBuffer should not have changed
+			CompressedData.SetNum(OutBufferSize);
+		}
 	}
 }
 
@@ -360,9 +422,10 @@ void FJpegImageWrapper::UncompressTurbo(const ERGBFormat InFormat, int32 InBitDe
 	else
 	{
 		check(false);
+		return;
 	}
 
-	FScopeLock JPEGLock(&GJPEGSection);
+	JPEG_SCOPE_CRITSEC();
 
 	check(Decompressor);
 	check(CompressedData.Num());
@@ -376,6 +439,7 @@ void FJpegImageWrapper::UncompressTurbo(const ERGBFormat InFormat, int32 InBitDe
 	if (tjDecompress2(Decompressor, CompressedData.GetData(), CompressedData.Num(), RawData.GetData(), Width, 0, Height, PixelFormat, Flags) != 0)
 	{
 		UE_LOG(LogImageWrapper, Error, TEXT("JPEG Decompress Error"));
+		SetError(TEXT("tjDecompress2 failed"));
 		RawData.Empty();
 		return;
 	}
diff --git a/Engine/Source/Runtime/ImageWrapper/Private/Formats/JpegImageWrapper.h b/Engine/Source/Runtime/ImageWrapper/Private/Formats/JpegImageWrapper.h
--- a/Engine/Source/Runtime/ImageWrapper/Private/Formats/JpegImageWrapper.h
+++ b/Engine/Source/Runtime/ImageWrapper/Private/Formats/JpegImageWrapper.h
@@ -22,11 +22,11 @@ class FJpegImageWrapper
 public:
 
 	/** Default constructor. */
+	// NumComponents == 1 for GrayscaleJPEG , but not supported by TJPEG path so currently broken
 	FJpegImageWrapper(int32 InNumComponents = 4);
 
-#if WITH_LIBJPEGTURBO
 	virtual ~FJpegImageWrapper();
-#endif	// WITH_LIBJPEGTURBO
+	virtual void Reset() override;
 
 public:
 
@@ -50,7 +50,6 @@ private:
 #endif	// WITH_LIBJPEGTURBO
 
 #if WITH_LIBJPEGTURBO
-	tjhandle Compressor;
 	tjhandle Decompressor;
 #endif	// WITH_LIBJPEGTURBO
 };
