fix ImageWrapperHelper identifying jpeg as grayscale
ImageWrapperHelper is deprecated, do not use

#rnx

diff --git a/Engine/Source/Runtime/ImageWrapper/Private/ImageWrapperHelper.cpp b/Engine/Source/Runtime/ImageWrapper/Private/ImageWrapperHelper.cpp
--- a/Engine/Source/Runtime/ImageWrapper/Private/ImageWrapperHelper.cpp
+++ b/Engine/Source/Runtime/ImageWrapper/Private/ImageWrapperHelper.cpp
@@ -9,19 +9,15 @@ FStringView ImageWrapperHelper::GetFormatExtension(EImageFormat InImageFormat, b
 	switch (InImageFormat)
 	{
 	case EImageFormat::Invalid:
-		return StringExtension;
+		return FStringView(TEXT(""));
 	case EImageFormat::PNG:
 		StringExtension = TEXT(".png");
 		break;
 	case EImageFormat::JPEG:
-		StringExtension = TEXT(".jpeg");
-		break;
 	case EImageFormat::GrayscaleJPEG:
 		StringExtension = TEXT(".jpg");
 		break;
 	case EImageFormat::OOJPEG:
-		StringExtension = TEXT(".ooj");
-		break;
 	case EImageFormat::GrayscaleOOJPEG:
 		StringExtension = TEXT(".ooj");
 		break;
@@ -44,7 +40,8 @@ FStringView ImageWrapperHelper::GetFormatExtension(EImageFormat InImageFormat, b
 		StringExtension = TEXT(".hdr");
 		break;
 	default:
-		return StringExtension;
+		// this should log error
+		return FStringView(TEXT(""));
 	}
 	if (!bIncludeDot) 
 	{
@@ -55,6 +52,8 @@ FStringView ImageWrapperHelper::GetFormatExtension(EImageFormat InImageFormat, b
 
 EImageFormat ImageWrapperHelper::GetImageFormat(FStringView StringExtention)
 {
+	// deprecated, use ImageWrapper::GetImageFormatFromExtension instead
+
 	int32 Length = StringExtention.Len();
 	if (Length == 0)
 	{
@@ -76,7 +75,7 @@ EImageFormat ImageWrapperHelper::GetImageFormat(FStringView StringExtention)
 	}
 	if (StringExtention.Equals(TEXT("jpg"), ESearchCase::IgnoreCase))
 	{
-		return EImageFormat::GrayscaleJPEG; // Note: This looks like a bug?
+		return EImageFormat::JPEG;
 	}
 	if (StringExtention.Equals(TEXT("ooj"), ESearchCase::IgnoreCase))
 	{
diff --git a/Engine/Source/Runtime/ImageWrapper/Public/ImageWrapperHelper.h b/Engine/Source/Runtime/ImageWrapper/Public/ImageWrapperHelper.h
--- a/Engine/Source/Runtime/ImageWrapper/Public/ImageWrapperHelper.h
+++ b/Engine/Source/Runtime/ImageWrapper/Public/ImageWrapperHelper.h
@@ -5,10 +5,16 @@
 #include "CoreMinimal.h"
 #include "IImageWrapper.h"
 
+// ImageWrapperHelper is deprecated, use IImageWrapperModule instead
 class ImageWrapperHelper
 {
 public:
-	 static IMAGEWRAPPER_API FStringView GetFormatExtension(EImageFormat InImageFormat, bool bIncludeDot=false);
-	 static IMAGEWRAPPER_API EImageFormat GetImageFormat(FStringView StringExtention);
-	 static IMAGEWRAPPER_API const FStringView GetImageFilesFilterString(bool bIncludeAllFiles);
+
+	// deprecated, use IImageWrapperModule::GetExtension() instead
+	static IMAGEWRAPPER_API FStringView GetFormatExtension(EImageFormat InImageFormat, bool bIncludeDot=false);
+
+	// deprecated, use IImageWrapperModule::GetImageFormatFromExtension instead
+	static IMAGEWRAPPER_API EImageFormat GetImageFormat(FStringView StringExtention);
+
+	static IMAGEWRAPPER_API const FStringView GetImageFilesFilterString(bool bIncludeAllFiles);
 };
