Oodle JPEG: FASTDCT flag which matches libjpegturbo's FASTDCT flag. Enabled for now so output matches.
#p4v-preflight-copy 29024880

diff --git a/Engine/Source/Runtime/ImageWrapper/Private/Formats/OOJpegImageWrapper.cpp b/Engine/Source/Runtime/ImageWrapper/Private/Formats/OOJpegImageWrapper.cpp
--- a/Engine/Source/Runtime/ImageWrapper/Private/Formats/OOJpegImageWrapper.cpp
+++ b/Engine/Source/Runtime/ImageWrapper/Private/Formats/OOJpegImageWrapper.cpp
@@ -80,7 +80,7 @@ bool FOOJpegImageWrapper::SetCompressed(const void* InCompressedData, int64 InCo
 
 		// Then decode to a raw image
 		{
-			oojpeg_decode_context_t ctx = oojpeg_decode_mem_threaded_start(OojData, OojDataSize, 0);
+			oojpeg_decode_context_t ctx = oojpeg_decode_mem_threaded_start(OojData, OojDataSize, OOJPEG_FLAG_FASTDCT);
 			if (ctx.error)
 			{
 				FMemory::Free(OojData);
@@ -108,7 +108,7 @@ bool FOOJpegImageWrapper::SetCompressed(const void* InCompressedData, int64 InCo
 	}
 	else 
 	{
-		oojpeg_decode_context_t ctx = oojpeg_decode_mem_threaded_start((uint8*)InCompressedData, (uint32)InCompressedSize, 0);
+		oojpeg_decode_context_t ctx = oojpeg_decode_mem_threaded_start((uint8*)InCompressedData, (uint32)InCompressedSize, OOJPEG_FLAG_FASTDCT);
 		if (ctx.error)
 		{
 			return false;
@@ -258,7 +258,7 @@ void FOOJpegImageWrapper::Uncompress(const ERGBFormat InFormat, int32 InBitDepth
 	int32 NumColors;
 	int32 jpegWidth,jpegHeight;
 	uint8* OutData;
-	oojpeg_decode_context_t ctx = oojpeg_decode_mem_threaded_start(CompressedData.GetData(), CompressedData.Num(), 0);
+	oojpeg_decode_context_t ctx = oojpeg_decode_mem_threaded_start(CompressedData.GetData(), CompressedData.Num(), OOJPEG_FLAG_FASTDCT);
 	if (ctx.error)
 	{
 		return;
diff --git a/Engine/Source/Runtime/OodleJPEG/Source/oojpeg.h b/Engine/Source/Runtime/OodleJPEG/Source/oojpeg.h
--- a/Engine/Source/Runtime/OodleJPEG/Source/oojpeg.h
+++ b/Engine/Source/Runtime/OodleJPEG/Source/oojpeg.h
@@ -93,6 +93,7 @@ typedef enum {
 typedef enum {
   OOJPEG_FLAG_NONE            = 0x0, // no flags
   OOJPEG_FLAG_NO_COMPRESSION  = 0x1, // don't compress the coefficients, just store them
+  OOJPEG_FLAG_FASTDCT         = 0x2, // use a fast dct transform which matches libjpegturbo's fastdct.
 } oojpeg_flags_t;
 
 enum {
