ImageWrapperBase : fix use of int should be int64

#rnx

diff --git a/Engine/Source/Runtime/ImageWrapper/Private/ImageWrapperBase.cpp b/Engine/Source/Runtime/ImageWrapper/Private/ImageWrapperBase.cpp
--- a/Engine/Source/Runtime/ImageWrapper/Private/ImageWrapperBase.cpp
+++ b/Engine/Source/Runtime/ImageWrapper/Private/ImageWrapperBase.cpp
@@ -117,9 +117,9 @@ bool FImageWrapperBase::SetRaw(const void* InRawData, int64 InRawSize, const int
 	Width = InWidth;
 	Height = InHeight;
 
-	int BytesPerRow = GetBytesPerRow();
+	int64 BytesPerRow = GetBytesPerRow();
 	
-	int64 RawSize = (int64) BytesPerRow * Height;
+	int64 RawSize = BytesPerRow * Height;
 	RawData.Empty(RawSize);
 	RawData.AddUninitialized(RawSize);
 
@@ -396,10 +396,10 @@ bool IImageWrapper::GetRawImage(FImage & OutImage)
 		return false;
 	}
 
-	int Width = GetWidth();
-	int Height = GetHeight();
+	int64 Width = GetWidth();
+	int64 Height = GetHeight();
 	ERGBFormat RGBFormat = GetFormat();
-	int BitDepth = GetBitDepth();
+	int32 BitDepth = GetBitDepth();
 
 	bool bExactMatch;
 	ERawImageFormat::Type RawFormat = GetClosestRawImageFormat(&bExactMatch);
