Route png compress through user memory pointers.

diff --git a/Engine/Source/Runtime/ImageWrapper/Private/Formats/PngImageWrapper.cpp b/Engine/Source/Runtime/ImageWrapper/Private/Formats/PngImageWrapper.cpp
--- a/Engine/Source/Runtime/ImageWrapper/Private/Formats/PngImageWrapper.cpp
+++ b/Engine/Source/Runtime/ImageWrapper/Private/Formats/PngImageWrapper.cpp
@@ -167,8 +167,8 @@ void FPngImageWrapper::Compress(int32 Quality)
 
 		// Reset to the beginning of file so we can use png_read_png(), which expects to start at the beginning.
 		ReadOffset = 0;
-
-		png_structp png_ptr	= png_create_write_struct(PNG_LIBPNG_VER_STRING, this, FPngImageWrapper::user_error_fn, FPngImageWrapper::user_warning_fn);
+		
+		png_structp png_ptr	= png_create_write_struct_2(PNG_LIBPNG_VER_STRING, this, FPngImageWrapper::user_error_fn, FPngImageWrapper::user_warning_fn, NULL, FPngImageWrapper::user_malloc, FPngImageWrapper::user_free);
 		check(png_ptr);
 
 		png_infop info_ptr	= png_create_info_struct(png_ptr);
