Explicit error message when importing jpeg with unsupported CMYK colorspace.

diff --git a/Engine/Source/Runtime/ImageWrapper/Private/Formats/JpegImageWrapper.cpp b/Engine/Source/Runtime/ImageWrapper/Private/Formats/JpegImageWrapper.cpp
--- a/Engine/Source/Runtime/ImageWrapper/Private/Formats/JpegImageWrapper.cpp
+++ b/Engine/Source/Runtime/ImageWrapper/Private/Formats/JpegImageWrapper.cpp
@@ -336,6 +336,12 @@ bool FJpegImageWrapper::SetCompressedTurbo(const void* InCompressedData, int64 I
 		return false;
 	}
 
+	if (ColorSpace == TJCS_CMYK || ColorSpace == TJCS_YCCK)
+	{
+		UE_LOG(LogImageWrapper, Error, TEXT("Unsupported CMYK colorspace for JPEG, only Grayscale and RGB are supported"));
+		return false;
+	}
+
 	// set after call to base SetCompressed as it will reset members
 	Width = ImageWidth;
 	Height = ImageHeight;
@@ -438,7 +444,7 @@ void FJpegImageWrapper::UncompressTurbo(const ERGBFormat InFormat, int32 InBitDe
 
 	if (tjDecompress2(Decompressor, CompressedData.GetData(), CompressedData.Num(), RawData.GetData(), Width, 0, Height, PixelFormat, Flags) != 0)
 	{
-		UE_LOG(LogImageWrapper, Error, TEXT("JPEG Decompress Error"));
+		UE_LOG(LogImageWrapper, Error, TEXT("JPEG Decompress Error - %s"), ANSI_TO_TCHAR(tj3GetErrorStr(Decompressor)));
 		SetError(TEXT("tjDecompress2 failed"));
 		RawData.Empty();
 		return;
