HdrImageWrapper : remove arbitrary 65536 dimension limit
correct place to put limits is in FImageCoreUtils::IsImageImportPossible

diff --git a/Engine/Source/Runtime/ImageWrapper/Private/Formats/HdrImageWrapper.cpp b/Engine/Source/Runtime/ImageWrapper/Private/Formats/HdrImageWrapper.cpp
--- a/Engine/Source/Runtime/ImageWrapper/Private/Formats/HdrImageWrapper.cpp
+++ b/Engine/Source/Runtime/ImageWrapper/Private/Formats/HdrImageWrapper.cpp
@@ -92,17 +92,10 @@ bool FHdrImageWrapper::SetCompressedFromView(TArrayView64<const uint8> Data)
 		return FailHeaderParsing();
 	}
 
-	// If we allow enormous images, we risk int overflows on the size calcs, even with int64s.
-	// As of this writing (Oct 2022) there is no demand for reading huge HDR files, and allocating
-	// memory to enormous images is going to make the editor choke besides. Limiting
-	// pixel counts to 65536*65536 (which caps images at 16GiB) seems reasonable for now.
-	const int MaxImageDimension = 65536;
-
 	int ImageWidth;
 	int ImageHeight;
-	if (!ParseImageSize(Line, &ImageWidth, &ImageHeight) ||
-		ImageWidth <= 0 || ImageWidth > MaxImageDimension ||
-		ImageHeight <= 0 || ImageHeight > MaxImageDimension)
+	if ( !ParseImageSize(Line, &ImageWidth, &ImageHeight) ||
+		ImageWidth <= 0 || ImageHeight <= 0 )
 	{
 		// If we don't like the resolution line (our parser is very strict), log what it was
 		// as a breadcrumb for debugging.
