Add support for importing .tx texture file format

- Added overloads for ImageWrappers to add meta information about sub images
- TX images are handled through libTIFF.
- TiffImageWrapper handles more than one directories in the TIFF file

diff --git a/Engine/Plugins/Interchange/Runtime/Source/Import/Private/Texture/InterchangeImageWrapperTranslator.cpp b/Engine/Plugins/Interchange/Runtime/Source/Import/Private/Texture/InterchangeImageWrapperTranslator.cpp
--- a/Engine/Plugins/Interchange/Runtime/Source/Import/Private/Texture/InterchangeImageWrapperTranslator.cpp
+++ b/Engine/Plugins/Interchange/Runtime/Source/Import/Private/Texture/InterchangeImageWrapperTranslator.cpp
@@ -6,6 +6,7 @@
 #include "Engine/Texture2D.h"
 #include "IImageWrapper.h"
 #include "IImageWrapperModule.h"
+#include "ImageWrapperOutputTypes.h"
 #include "ImageCoreUtils.h"
 #include "InterchangeImportLog.h"
 #include "InterchangeTextureNode.h"
@@ -60,6 +61,14 @@ static FAutoConsoleVariableRef CCvarInterchangeEnableTIFFImport(
 	ECVF_Default);
 #endif
 
+// TODO: should this still be a feature that can be turned off? 
+static bool GInterchangeEnableMipMapImageImport = true;
+static FAutoConsoleVariableRef CCvarInterchangeEnableMipMapImageImport(
+	TEXT("Interchange.FeatureFlags.Import.MipMapImage"),
+	GInterchangeEnableMipMapImageImport,
+	TEXT("Whether Mip Mapped Image support is enabled."),
+	ECVF_Default);
+
 static bool GInterchangeEnableTGAImport = true;
 static FAutoConsoleVariableRef CCvarInterchangeEnableTGAImport(
 	TEXT("Interchange.FeatureFlags.Import.TGA"),
@@ -67,6 +76,21 @@ static FAutoConsoleVariableRef CCvarInterchangeEnableTGAImport(
 	TEXT("Whether TGA support is enabled."),
 	ECVF_Default);
 
+
+namespace UE::Interchange::ImageWrapperTranslator::Private
+{
+	static bool SupportsMipMapsAndMetaData(EImageFormat ImageFormat)
+	{
+		switch (ImageFormat)
+		{
+		case EImageFormat::TIFF:
+			return true;
+		default:
+			return false;
+		}
+	}
+}
+
 UInterchangeImageWrapperTranslator::UInterchangeImageWrapperTranslator()
 {
 	// construction of CDO is not thread safe,
@@ -106,6 +130,7 @@ TArray<FString> UInterchangeImageWrapperTranslator::GetSupportedFormats() const
 	{
 		Formats.Emplace(TEXT("tif;Tag Image File Format"));
 		Formats.Emplace(TEXT("tiff;Tag Image File Format"));
+		Formats.Emplace(TEXT("tx;Tag Image File Format"));
 	}
 #endif
 
@@ -159,79 +184,125 @@ TOptional<UE::Interchange::FImportImage> UInterchangeImageWrapperTranslator::Get
 
 	if (ImageFormat != EImageFormat::Invalid)
 	{
-		// Generic ImageWrapper loader :
-		// for PNG,EXR,BMP,TGA :
-		FImage LoadedImage;
-		if (ImageWrapperModule.DecompressImage(Buffer, Length, LoadedImage))
+		using namespace UE::Interchange::ImageWrapperTranslator::Private;
+
+		if (GInterchangeEnableMipMapImageImport && SupportsMipMapsAndMetaData(ImageFormat))
 		{
-			// Todo interchange: should these payload modification be part of the pipeline, factory or stay there?
-			if (UE::TextureUtilitiesCommon::AutoDetectAndChangeGrayScale(LoadedImage))
+			FDecompressedImageOutput DecompressedImage;
+			if (ImageWrapperModule.DecompressImage(Buffer, Length, DecompressedImage))
 			{
-				UE_LOG(LogInterchangeImport, Display, TEXT("Auto-detected grayscale, image changed to G8"));
-			}
-
-			ETextureSourceFormat TextureFormat = FImageCoreUtils::ConvertToTextureSourceFormat(LoadedImage.Format);
-			bool bSRGB = LoadedImage.GammaSpace != EGammaSpace::Linear;
+				if (UE::TextureUtilitiesCommon::AutoDetectAndChangeGrayScale(DecompressedImage.MipMapImage))
+				{
+					UE_LOG(LogInterchangeImport, Display, TEXT("Auto-detected grayscale, image changed to G8"));
+				}
 
-			PayloadData.Init2DWithParams(
-				LoadedImage.SizeX,
-				LoadedImage.SizeY,
-				TextureFormat,
-				bSRGB,
-				false
-			);
+				ETextureSourceFormat TextureFormat = FImageCoreUtils::ConvertToTextureSourceFormat(DecompressedImage.MipMapImage.Format);
+				bool bSRGB = DecompressedImage.MipMapImage.GammaSpace != EGammaSpace::Linear;
+				
+				constexpr int32 MipLevel = 0;
+				FImageView MipZeroImageView = DecompressedImage.MipMapImage.GetMipImage(MipLevel);
 
-			PayloadData.RawData = MakeUniqueBufferFromArray(MoveTemp(LoadedImage.RawData));
+				PayloadData.Init2DWithParams(
+					MipZeroImageView.SizeX,
+					MipZeroImageView.SizeY,
+					TextureFormat,
+					bSRGB,
+					false
+				);
 
-			if (ERawImageFormat::IsHDR(LoadedImage.Format))
-			{
-				PayloadData.CompressionSettings = TC_HDR;
-				check(bSRGB == false);
-			}
+				PayloadData.RawData = MakeUniqueBufferFromArray(MoveTemp(DecompressedImage.MipMapImage.RawData));
 
-			// do per-format processing to match legacy behavior :
-
-			if (ImageFormat == EImageFormat::PNG)
-			{
-				if (PNGInfill != ETextureImportPNGInfill::Never)
+				if (ERawImageFormat::IsHDR(DecompressedImage.MipMapImage.Format))
 				{
-					bool bDoOnComplexAlphaNotJustBinaryTransparency = ( PNGInfill == ETextureImportPNGInfill::Always );
+					PayloadData.CompressionSettings = TC_HDR;
+					check(bSRGB == false);
+				}
 
-					// Replace the pixels with 0.0 alpha with a color value from the nearest neighboring color which has a non-zero alpha
-					UE::TextureUtilitiesCommon::FillZeroAlphaPNGData(PayloadData.SizeX, PayloadData.SizeY, PayloadData.Format, reinterpret_cast<uint8*>(PayloadData.RawData.GetData()), bDoOnComplexAlphaNotJustBinaryTransparency);
+				// Format Specific Settings
+				if (ImageFormat == EImageFormat::TIFF)
+				{
+					PayloadData.NumMips = DecompressedImage.MipMapImage.GetMipCount();
+					PayloadData.MipGenSettings = PayloadData.NumMips > 1 ? TextureMipGenSettings::TMGS_LeaveExistingMips : TextureMipGenSettings::TMGS_NoMipmaps;
 				}
 			}
-			else if (ImageFormat == EImageFormat::TGA)
-			{
-				const FTGAFileHeader* TGA = (FTGAFileHeader*)Buffer;
+		}
+		else
+		{
 
-				if (TGA->ColorMapType == 1 && TGA->ImageTypeCode == 1 && TGA->BitsPerPixel == 8)
+			// Generic ImageWrapper loader :
+			// for PNG,EXR,BMP,TGA :
+			FImage LoadedImage;
+			if (ImageWrapperModule.DecompressImage(Buffer, Length, LoadedImage))
+			{
+				// Todo interchange: should these payload modification be part of the pipeline, factory or stay there?
+				if (UE::TextureUtilitiesCommon::AutoDetectAndChangeGrayScale(LoadedImage))
 				{
-					// Notes: The Scaleform GFx exporter (dll) strips all font glyphs into a single 8-bit texture.
-					// The targa format uses this for a palette index; GFx uses a palette of (i,i,i,i) so the index
-					// is also the alpha value.
-					//
-					// We store the image as PF_G8, where it will be used as alpha in the Glyph shader.
+					UE_LOG(LogInterchangeImport, Display, TEXT("Auto-detected grayscale, image changed to G8"));
+				}
 
-					// ?? check or convert? or neither?
-					//check( TextureFormat == TSF_G8 );
+				ETextureSourceFormat TextureFormat = FImageCoreUtils::ConvertToTextureSourceFormat(LoadedImage.Format);
+				bool bSRGB = LoadedImage.GammaSpace != EGammaSpace::Linear;
 
-					PayloadData.CompressionSettings = TC_Grayscale;
-				}
-				else if (TGA->ColorMapType == 0 && TGA->ImageTypeCode == 3 && TGA->BitsPerPixel == 8)
-				{
-					// standard grayscale images
+				PayloadData.Init2DWithParams(
+					LoadedImage.SizeX,
+					LoadedImage.SizeY,
+					TextureFormat,
+					bSRGB,
+					false
+				);
 
-					// ?? check or convert? or neither?
-					//check( TextureFormat == TSF_G8 );
+				PayloadData.RawData = MakeUniqueBufferFromArray(MoveTemp(LoadedImage.RawData));
 
-					PayloadData.CompressionSettings = TC_Grayscale;
+				if (ERawImageFormat::IsHDR(LoadedImage.Format))
+				{
+					PayloadData.CompressionSettings = TC_HDR;
+					check(bSRGB == false);
 				}
 
-				if (PayloadData.CompressionSettings == TC_Grayscale && TGA->ImageTypeCode == 3)
+				// do per-format processing to match legacy behavior :
+
+				if (ImageFormat == EImageFormat::PNG)
+				{
+					if (PNGInfill != ETextureImportPNGInfill::Never)
+					{
+						bool bDoOnComplexAlphaNotJustBinaryTransparency = (PNGInfill == ETextureImportPNGInfill::Always);
+
+						// Replace the pixels with 0.0 alpha with a color value from the nearest neighboring color which has a non-zero alpha
+						UE::TextureUtilitiesCommon::FillZeroAlphaPNGData(PayloadData.SizeX, PayloadData.SizeY, PayloadData.Format, reinterpret_cast<uint8*>(PayloadData.RawData.GetData()), bDoOnComplexAlphaNotJustBinaryTransparency);
+					}
+				}
+				else if (ImageFormat == EImageFormat::TGA)
 				{
-					// default grayscales to linear as they wont get compression otherwise and are commonly used as masks
-					PayloadData.bSRGB = false;
+					const FTGAFileHeader* TGA = (FTGAFileHeader*)Buffer;
+
+					if (TGA->ColorMapType == 1 && TGA->ImageTypeCode == 1 && TGA->BitsPerPixel == 8)
+					{
+						// Notes: The Scaleform GFx exporter (dll) strips all font glyphs into a single 8-bit texture.
+						// The targa format uses this for a palette index; GFx uses a palette of (i,i,i,i) so the index
+						// is also the alpha value.
+						//
+						// We store the image as PF_G8, where it will be used as alpha in the Glyph shader.
+
+						// ?? check or convert? or neither?
+						//check( TextureFormat == TSF_G8 );
+
+						PayloadData.CompressionSettings = TC_Grayscale;
+					}
+					else if (TGA->ColorMapType == 0 && TGA->ImageTypeCode == 3 && TGA->BitsPerPixel == 8)
+					{
+						// standard grayscale images
+
+						// ?? check or convert? or neither?
+						//check( TextureFormat == TSF_G8 );
+
+						PayloadData.CompressionSettings = TC_Grayscale;
+					}
+
+					if (PayloadData.CompressionSettings == TC_Grayscale && TGA->ImageTypeCode == 3)
+					{
+						// default grayscales to linear as they wont get compression otherwise and are commonly used as masks
+						PayloadData.bSRGB = false;
+					}
 				}
 			}
 		}
diff --git a/Engine/Source/Runtime/ImageCore/Private/ImageCore.cpp b/Engine/Source/Runtime/ImageCore/Private/ImageCore.cpp
--- a/Engine/Source/Runtime/ImageCore/Private/ImageCore.cpp
+++ b/Engine/Source/Runtime/ImageCore/Private/ImageCore.cpp
@@ -1386,6 +1386,86 @@ void FImage::Linearize(uint8 SourceEncoding, FImage& DestImage) const
 	}
 }
 
+
+
+void FMipMapImage::Init(ERawImageFormat::Type InFormat, EGammaSpace InGammaSpace)
+{
+	Format = InFormat;
+	GammaSpace = InGammaSpace;
+	SubImages.Empty();
+	RawData.Empty();
+}
+
+void FMipMapImage::Init(int32 MipZeroWidth, int32 MipZeroHeight, int32 NumMips, ERawImageFormat::Type InFormat, EGammaSpace InGammaSpace)
+{
+	Format = InFormat;
+	GammaSpace = InGammaSpace;
+
+	const int32 BytesPerPixel = ERawImageFormat::GetBytesPerPixel(Format);
+	int64 BufferSizeBytes = 0;
+
+	int32 CurrMipWidth = MipZeroWidth;
+	int32 CurrMipHeight = MipZeroHeight;
+
+	SubImages.Empty(NumMips);
+	
+	do 
+	{
+		const int32 MipSizeBytes = CurrMipWidth * CurrMipHeight * BytesPerPixel;
+		SubImages.Add(FMipInfo{ CurrMipWidth, CurrMipHeight, BufferSizeBytes, MipSizeBytes });
+
+		BufferSizeBytes += MipSizeBytes;
+		CurrMipWidth = CurrMipWidth >> 1;
+		CurrMipHeight = CurrMipHeight >> 1;
+		NumMips--;
+	} while ((NumMips > 0) && (CurrMipWidth > 0) && (CurrMipHeight > 0));
+
+	RawData.Empty(BufferSizeBytes);
+	RawData.AddUninitialized(BufferSizeBytes);
+}
+
+void FMipMapImage::CopyTo(FMipMapImage& DestImage, ERawImageFormat::Type DestFormat, EGammaSpace DestGammaSpace)
+{
+	// if gamma correction is done, it's always *TO* sRGB , not to Pow22
+	// so if Pow22 was requested, change to sRGB
+	// so that Float->int->Float roundtrips correctly
+	if (DestGammaSpace == EGammaSpace::Pow22 && GammaSpace != EGammaSpace::Pow22)
+	{
+		// fix call sites that hit this
+		UE_LOG(LogImageCore, Warning, TEXT("Pow22 should not be used as a Dest GammaSpace.  Pow22 Source should encode to sRGB Dest."));
+		DestGammaSpace = EGammaSpace::sRGB;
+	}
+
+	constexpr int32 MipLevel = 0;
+	int32 Width, Height;
+	if (GetMipDimensions(MipLevel, Width, Height))
+	{
+		// existing contents of DestImage are freed and replaced
+		DestImage.Init(Width, Height, GetMipCount(), DestFormat, DestGammaSpace);
+		ParallelFor(GetMipCount(), [this, &DestImage](int32 MipLevel) 
+		{
+			FImageView SrcMipImageView = GetMipImage(MipLevel);
+			FImageView DestMipImageView = DestImage.GetMipImage(MipLevel);
+			FImageCore::CopyImage(SrcMipImageView, DestMipImageView);
+		});
+	}
+}
+
+void FMipMapImage::ChangeFormat(ERawImageFormat::Type DestFormat, EGammaSpace DestGammaSpace)
+{
+	if (Format == DestFormat &&
+		(GammaSpace == DestGammaSpace || !ERawImageFormat::GetFormatNeedsGammaSpace(Format)))
+	{
+		// no action needed
+	}
+	else
+	{
+		FMipMapImage Temp;
+		CopyTo(Temp, DestFormat, DestGammaSpace);
+		::Swap(*this, Temp);
+	}
+}
+
 const FUtf8StringView ERawImageFormat::GetNameView(Type Format)
 {
 	switch (Format)
@@ -2331,3 +2411,5 @@ IMAGECORE_API void FImageCore::ResizeImageInPlace(FImage & Image,int32 DestSizeX
 
 //} resize
 //----------------------
+
+
diff --git a/Engine/Source/Runtime/ImageCore/Public/ImageCore.h b/Engine/Source/Runtime/ImageCore/Public/ImageCore.h
--- a/Engine/Source/Runtime/ImageCore/Public/ImageCore.h
+++ b/Engine/Source/Runtime/ImageCore/Public/ImageCore.h
@@ -715,6 +715,135 @@ struct FSharedImage : public FImage, public FThreadSafeRefCountedObject
 	virtual ~FSharedImage() = default;
 };
 
+/**
+ * Structure for raw image data used for Image formats that support Mip Maps.
+ */
+struct FMipMapImage
+{
+	struct FMipInfo
+	{
+		/** Size of the mip image in Pixels*/
+		int32 Width = 0;
+
+		/** Size of the mip image in Pixels*/
+		int32 Height = 0;
+
+		/** Offset into RawData */
+		int64 Offset = 0;
+
+		/** Size of the mip image in Bytes*/
+		int64 Size = 0;
+	};
+
+	/** Raw image data. */
+	TArray64<uint8> RawData;
+
+	/** Holds information about the sub images some image formats support. */
+	TArray<FMipInfo> SubImages;
+
+	/** Format in which the images is stored. */
+	ERawImageFormat::Type Format = ERawImageFormat::BGRA8;
+
+	/** The gamma space the image is stored in. */
+	EGammaSpace GammaSpace = EGammaSpace::sRGB;
+
+public:
+	/**
+	 * Initializes the Mip Map Image with parameters with uninitialized SubImages.
+	 * 
+	 * @param InFormat Raw Image Format
+	 * @param InGammaSpace Gamma Space
+	 */
+	IMAGECORE_API void Init(ERawImageFormat::Type InFormat, EGammaSpace InGammaSpace);
+
+	/**
+	 * Initializes the Mip Map Image with reserving space for mip maps.
+	 *
+	 * @param InFormat Raw Image Format
+	 * @param InGammaSpace Gamma Space
+	 */
+	IMAGECORE_API void Init(int32 MipZeroWidth, int32 MipZeroHeight, int32 NumMips, ERawImageFormat::Type InFormat, EGammaSpace InGammaSpace);
+
+	/**
+	 * Copies the image to a destination image with the specified format.
+	 *
+	 * @param DestImage - The destination image.  Will be allocated.  Any existing contents are replaced.
+	 * @param DestFormat - The destination image format.
+	 * @param DestSRGB - Whether the destination image is in SRGB format.
+	 */
+	IMAGECORE_API void CopyTo(FMipMapImage& DestImage, ERawImageFormat::Type DestFormat, EGammaSpace DestGammaSpace);
+
+	void CopyTo(FMipMapImage& DestImage)
+	{
+		CopyTo(DestImage, Format, GammaSpace);
+	}
+
+	/**
+	 * In - place format change
+	 * does nothing if already in the desired format
+	 */ 
+	IMAGECORE_API void ChangeFormat(ERawImageFormat::Type DestFormat, EGammaSpace DestGammaSpace);
+
+
+	FORCEINLINE bool IsValid() const
+	{
+		bool bMipInfoValid = true;
+		for (const FMipInfo& MipInfo : SubImages)
+		{
+			bMipInfoValid &= (MipInfo.Width > 0 && MipInfo.Height > 0 && MipInfo.Size > 0);
+		}
+
+		if (!bMipInfoValid) return false;
+		if (Format == ERawImageFormat::Invalid) return false;
+		if (GammaSpace == EGammaSpace::Invalid) return false;
+		if (!GetFormatNeedsGammaSpace(Format) && GammaSpace != EGammaSpace::Linear) return false;
+		return true;
+	}
+
+	FORCEINLINE bool GetMipDimensions(int32 MipLevel, int32& OutWidth, int32& OutHeight) const
+	{
+		if (!SubImages.IsValidIndex(MipLevel))
+		{
+			OutWidth = 0;
+			OutHeight = 0;
+			return false;
+		}
+
+		OutWidth = SubImages[MipLevel].Width;
+		OutHeight = SubImages[MipLevel].Height;
+		return true;
+	}
+
+	FORCEINLINE int32 GetMipCount() const 
+	{
+		return SubImages.Num();
+	}
+
+	IMAGECORE_API FImageView GetMipImage(int32 MipLevel)
+	{
+		check(SubImages.IsValidIndex(MipLevel));
+		return FImageView(static_cast<void*>(RawData.GetData() + SubImages[MipLevel].Offset), SubImages[MipLevel].Width, SubImages[MipLevel].Height, 1, Format, GammaSpace);
+	}
+
+	IMAGECORE_API void AddMipImage(TArray64<uint8>&& Buffer, int32 Width, int32 Height)
+	{
+		FMipInfo MipInfo;
+		{
+			MipInfo.Width = Width;
+			MipInfo.Height = Height;
+			MipInfo.Offset = RawData.Num();
+			MipInfo.Size = Width * Height * ERawImageFormat::GetBytesPerPixel(Format);
+		}
+		SubImages.Add(MipInfo);
+		RawData.Append(MoveTemp(Buffer));
+	}
+
+	FORCEINLINE int64 GetNumPixels()
+	{
+		return (RawData.Num() / ERawImageFormat::GetBytesPerPixel(Format));
+	}
+
+};
 
 
 /* Functions
diff --git a/Engine/Source/Runtime/ImageWrapper/Private/Formats/HdrImageWrapper.cpp b/Engine/Source/Runtime/ImageWrapper/Private/Formats/HdrImageWrapper.cpp
--- a/Engine/Source/Runtime/ImageWrapper/Private/Formats/HdrImageWrapper.cpp
+++ b/Engine/Source/Runtime/ImageWrapper/Private/Formats/HdrImageWrapper.cpp
@@ -234,6 +234,17 @@ bool FHdrImageWrapper::GetRaw(const ERGBFormat InFormat, int32 InBitDepth, TArra
 	return true;
 }
 
+IMAGEWRAPPER_API bool FHdrImageWrapper::GetRaw(const ERGBFormat InFormat, int32 InBitDepth, FDecompressedImageOutput& OutDecompressedImage)
+{
+	TArray64<uint8> OutRawBuffer;
+	if (!GetRaw(InFormat, InBitDepth, OutRawBuffer))
+	{
+		return false;
+	}
+	OutDecompressedImage.MipMapImage.AddMipImage(MoveTemp(OutRawBuffer), GetWidth(), GetHeight());
+	return true;
+}
+
 int64 FHdrImageWrapper::GetWidth() const
 {
 	return Width;
diff --git a/Engine/Source/Runtime/ImageWrapper/Private/Formats/TiffImageWrapper.cpp b/Engine/Source/Runtime/ImageWrapper/Private/Formats/TiffImageWrapper.cpp
--- a/Engine/Source/Runtime/ImageWrapper/Private/Formats/TiffImageWrapper.cpp
+++ b/Engine/Source/Runtime/ImageWrapper/Private/Formats/TiffImageWrapper.cpp
@@ -447,54 +447,117 @@ namespace UE::ImageWrapper::Private
 
 	void FTiffImageWrapper::Uncompress(const ERGBFormat InFormat, int32 InBitDepth)
 	{
-		if ( Tiff == nullptr )
+		ON_SCOPE_EXIT{
+			ReleaseTiffImage();
+		};
+
+		if (Tiff == nullptr)
 		{
 			SetError(TEXT("Tiff invalid."));
 			return;
 		}
 
+		// We just want mip 0
+		const bool bUncompressResult = Uncompress_Internal(InFormat, InBitDepth);
+		if (bUncompressResult && SubImageBuffer.Num())
+		{
+			// Append the current sub-image to the base byte buffer that reads all the 
+			RawData.Append(MoveTemp(SubImageBuffer));
+		}
+		else
+		{
+			Width = 0;
+			Height = 0;
+		}
+
+		CurrSubImageWidth = 0;
+		CurrSubImageHeight = 0;
+	}
+
+	void FTiffImageWrapper::Uncompress(const ERGBFormat InFormat, int32 InBitDepth, FDecompressedImageOutput& OutDecompressedImage)
+	{
+		ON_SCOPE_EXIT{
+			ReleaseTiffImage();
+		};
+
+		if (Tiff == nullptr)
+		{
+			SetError(TEXT("Tiff invalid."));
+			return;
+		}
+
+		do
+		{
+			const bool bUncompressResult = Uncompress_Internal(InFormat, InBitDepth);
+			if (bUncompressResult && SubImageBuffer.Num())
+			{
+				static bool bExactMatch = false;
+
+				const int64 SubImageOffset = RawData.Num();
+				// Append the current sub-image to the base byte buffer that reads all the 
+				OutDecompressedImage.MipMapImage.AddMipImage(MoveTemp(SubImageBuffer), CurrSubImageWidth, CurrSubImageHeight);
+			}
+			else
+			{
+				Width = 0;
+				Height = 0;
+				break;
+			}
+
+		} while (TIFFReadDirectory(Tiff)); // This will go to the next directory if available. Next Directory will contain the next sub-image.
+
+		CurrSubImageWidth = 0;
+		CurrSubImageHeight = 0;
+	}
+
+	bool FTiffImageWrapper::Uncompress_Internal(const ERGBFormat InFormat, int32 InBitDepth)
+	{
+		// SubImageBuffer gets emptied in the respective calls to UnpackIntoRawBuffer
+		TIFFGetField(Tiff, TIFFTAG_IMAGEWIDTH, &CurrSubImageWidth);
+		TIFFGetField(Tiff, TIFFTAG_IMAGELENGTH, &CurrSubImageHeight);
+
 		if (InFormat == Format && InBitDepth == BitDepth)
 		{
 			// Read using RGBA
 			if (Format == ERGBFormat::BGRA && BitDepth == 8)
 			{
-				FGuardedInt64 GuardedBufferSize = FGuardedInt64(Width) * Height * sizeof(uint32);
+				FGuardedInt64 GuardedBufferSize = FGuardedInt64(CurrSubImageWidth) * CurrSubImageHeight * sizeof(uint32);
 				if (GuardedBufferSize.IsValid() == false)
 				{
 					SetError(TEXT("Tiff image is massive - likely corrupted file."));
-					return;
+					return false;
 				}
 
 				const int64 BufferSize = GuardedBufferSize.Get(0);
 				const int Flags = 0;
-	
-				RawData.Empty(BufferSize);
-				RawData.AddUninitialized(BufferSize);
 
-				if (TIFFReadRGBAImageOriented(Tiff, Width, Height, static_cast<uint32*>(static_cast<void*>(RawData.GetData())), 0, ORIENTATION_LEFTTOP) != 0)
+				SubImageBuffer.Empty(BufferSize);
+				SubImageBuffer.AddUninitialized(BufferSize);
+
+				if (TIFFReadRGBAImageOriented(Tiff, CurrSubImageWidth, CurrSubImageHeight, static_cast<uint32*>(static_cast<void*>(SubImageBuffer.GetData())), ORIENTATION_LEFTTOP, 0) != 0)
 				{
 					// @@ use ImageCore TransposeImageRGBABGRA
 					EParallelForFlags ParallelForFlags = IsInGameThread() ? EParallelForFlags::None : EParallelForFlags::BackgroundPriority;
-					ParallelFor(Height, [this](int32 HeightIndex)
+					ParallelFor(CurrSubImageHeight, [this](int32 HeightIndex)
 						{
-							const int64 HeightOffset = HeightIndex * Width;
-							for (int32 WidthIndex = 0; WidthIndex < Width; WidthIndex++)
+							const int64 HeightOffset = HeightIndex * CurrSubImageWidth;
+							for (int32 WidthIndex = 0; WidthIndex < CurrSubImageWidth; WidthIndex++)
 							{
 								// Swap from RGBA to BGRA
-								const int64 CurrentPixelIndex = (WidthIndex + HeightOffset)* sizeof(uint32);
-								uint8 Red = RawData[CurrentPixelIndex];
-								RawData[CurrentPixelIndex] = RawData[CurrentPixelIndex + 2];
-								RawData[CurrentPixelIndex + 2] = Red;
+								const int64 CurrentPixelIndex = (WidthIndex + HeightOffset) * sizeof(uint32);
+								uint8 Red = SubImageBuffer[CurrentPixelIndex];
+								SubImageBuffer[CurrentPixelIndex] = SubImageBuffer[CurrentPixelIndex + 2];
+								SubImageBuffer[CurrentPixelIndex + 2] = Red;
 							}
 						}
-						, ParallelForFlags);
+					, ParallelForFlags);
 				}
-				else 
+				else
 				{
 					UnpackIntoRawBuffer<uint8>(4);
 				}
 			}
-			else if (Format == ERGBFormat::RGBA && BitDepth == 16 )
+			else if (Format == ERGBFormat::RGBA && BitDepth == 16)
 			{
 				UnpackIntoRawBuffer<uint16>(4);
 			}
@@ -525,14 +588,16 @@ namespace UE::ImageWrapper::Private
 			else
 			{
 				SetError(TEXT("Unsupported requested format for the input image. Can't uncompress the tiff image."));
+				return false;
 			}
 		}
 		else
 		{
 			SetError(TEXT("Unsupported requested format for the input image. Can't uncompress the tiff image."));
+			return false;
 		}
 
-		ReleaseTiffImage();
+		return true;
 	}
 
 	bool FTiffImageWrapper::SetCompressed(const void* InCompressedData, int64 InCompressedSize)
@@ -722,6 +787,8 @@ namespace UE::ImageWrapper::Private
 
 	}
 
+	
+
 	void FTiffImageWrapper::ReleaseTiffImage()
 	{
 		if (Tiff)
@@ -794,13 +861,13 @@ namespace UE::ImageWrapper::Private
 	template<class DataTypeDest>
 	void FTiffImageWrapper::UnpackIntoRawBuffer(const uint8 NumOfChannelDest)
 	{
-		const int64 PixelSize = sizeof(DataTypeDest) * NumOfChannelDest;;
-		const int64 RowSize = PixelSize * int64(Width);
-		const int64 BufferSize = RowSize * int64(Height);
+		const int64 PixelSize = sizeof(DataTypeDest) * NumOfChannelDest;
+		const int64 RowSize = PixelSize * int64(CurrSubImageWidth);
+		const int64 BufferSize = RowSize * int64(CurrSubImageHeight);
 		const int Flags = 0;
 	
-		RawData.Empty(BufferSize);
-		RawData.AddUninitialized(BufferSize);
+		SubImageBuffer.Empty(BufferSize);
+		SubImageBuffer.AddUninitialized(BufferSize);
 
 		const bool bIsTiled = TIFFIsTiled(Tiff) != 0;
 
@@ -912,8 +979,8 @@ namespace UE::ImageWrapper::Private
 				// it's unclear if you should invert against 1.f or something else?
 				UE_LOG(LogImageWrapper, Warning, TEXT("Tiff MINISWHITE floating point?  This is probably wrong."));
 
-				TArrayView64<DataTypeDest> FinalImage(static_cast<DataTypeDest*>(static_cast<void*>(RawData.GetData())), RawData.Num());
-				ParallelFor(Width * Height, [&FinalImage](int32 Index)
+				TArrayView64<DataTypeDest> FinalImage(static_cast<DataTypeDest*>(static_cast<void*>(SubImageBuffer.GetData())), SubImageBuffer.Num());
+				ParallelFor(CurrSubImageWidth* CurrSubImageHeight, [&FinalImage](int32 Index)
 					{
 						DataTypeDest& FinalValue = FinalImage[Index];
 						FinalValue = 1.f - FinalValue; // uses implicit operator float() on FFloat16
@@ -921,8 +988,8 @@ namespace UE::ImageWrapper::Private
 			}
 			else
 			{
-				TArrayView64<DataTypeDest> FinalImage(static_cast<DataTypeDest*>(static_cast<void*>(RawData.GetData())), RawData.Num());
-				ParallelFor(Width * Height, [&FinalImage](int32 Index)
+				TArrayView64<DataTypeDest> FinalImage(static_cast<DataTypeDest*>(static_cast<void*>(SubImageBuffer.GetData())), SubImageBuffer.Num());
+				ParallelFor(CurrSubImageWidth* CurrSubImageHeight, [&FinalImage](int32 Index)
 					{
 						DataTypeDest& FinalValue = FinalImage[Index];
 						FinalValue = TNumericLimits<DataTypeDest>::Max() - FinalValue;
@@ -939,7 +1006,7 @@ namespace UE::ImageWrapper::Private
 		const uint64 LineSizeScr = TIFFScanlineSize64(Tiff);
 		const uint64 StripByteSize = TIFFStripSize64(Tiff);
 
-		TArrayView64<DataTypeDest> WriteArray(static_cast<DataTypeDest*>(static_cast<void*>(RawData.GetData())), RawData.Num() / sizeof(DataTypeDest));
+		TArrayView64<DataTypeDest> WriteArray(static_cast<DataTypeDest*>(static_cast<void*>(SubImageBuffer.GetData())), SubImageBuffer.Num() / sizeof(DataTypeDest));
 
 		uint16 PlanarConfig = 0;
 		TIFFGetFieldDefaulted(Tiff, TIFFTAG_PLANARCONFIG, &PlanarConfig);
@@ -948,8 +1015,8 @@ namespace UE::ImageWrapper::Private
 		const uint8 StepDest = bIsPlanarConfigSepareted ? NumOfChannelDest : 1;
 		const uint8 NumberOfChannelInReadArray =  bIsPlanarConfigSepareted ? 1 : SamplesPerPixel;
 
-		int32 TileWidth = Width;
-		int32 TileHeight = Height;
+		int32 TileWidth = CurrSubImageWidth;
+		int32 TileHeight = CurrSubImageHeight;
 
 		if constexpr (bIsTiled)
 		{
@@ -967,7 +1034,7 @@ namespace UE::ImageWrapper::Private
 				(int32 NumberOfColumnRead, int32 NumberOfRowRead, TArray64<uint8>& ReadBuffer, uint8 SampleIndex, int32 BlockX, int32 BlockY)
 				{
 					TArrayView64<const DataTypeSrc> ReadArray(static_cast<DataTypeSrc*>(static_cast<void*>(ReadBuffer.GetData())), ReadBuffer.Num() / sizeof(DataTypeSrc));
-					uint64 CurrentOffset = int64(BlockY) * Width * NumOfChannelDest;
+					uint64 CurrentOffset = int64(BlockY) * CurrSubImageWidth * NumOfChannelDest;
 					if constexpr (bIsTiled)
 					{
 						CurrentOffset *= TileHeight;
@@ -984,7 +1051,7 @@ namespace UE::ImageWrapper::Private
 								{
 									const int32 PositionXInTile = PixelReadIndex % NumberOfColumnRead;
 									const int32 PositionYInTile = PixelReadIndex / NumberOfColumnRead;
-									WriteIndex = CurrentOffset + PositionXInTile * NumOfChannelDest + PositionYInTile * Width * NumOfChannelDest + SampleIndex;
+									WriteIndex = CurrentOffset + PositionXInTile * NumOfChannelDest + PositionYInTile * CurrSubImageWidth* NumOfChannelDest + SampleIndex;
 
 									// The end of tile in X can be some garbage that act as padding data so that all tiles are the same size
 									ReadIndex = PositionXInTile * NumberOfChannelInReadArray + PositionYInTile * TileWidth * NumberOfChannelInReadArray;
@@ -1029,11 +1096,11 @@ namespace UE::ImageWrapper::Private
 
 			for (uint8 SampleIndex = 0; SampleIndex < NumOfPlanes; SampleIndex++)
 			{ 
-				for (int32 Y = 0, TileY = 0; Y < Height; Y += TileHeight, ++TileY)
+				for (int32 Y = 0, TileY = 0; Y < CurrSubImageHeight; Y += TileHeight, ++TileY)
 				{
-					const int32 NumberOfRow = Y + TileHeight > Height ? Height - Y : TileHeight;
+					const int32 NumberOfRow = Y + TileHeight > CurrSubImageHeight ? CurrSubImageHeight - Y : TileHeight;
 
-					for (int32 X = 0, TileX = 0; X < Width; X += TileWidth, ++TileX)
+					for (int32 X = 0, TileX = 0; X < CurrSubImageWidth; X += TileWidth, ++TileX)
 					{
 						TSharedPtr<TArray64<uint8>, ESPMode::ThreadSafe> BufferPtr;
 						if (!UsableBufferQueue.Dequeue(BufferPtr))
@@ -1048,7 +1115,7 @@ namespace UE::ImageWrapper::Private
 							return false;
 						}
 
-						const int32 NumberOfColumn = X + TileWidth > Width ? Width - X : TileWidth;
+						const int32 NumberOfColumn = X + TileWidth > CurrSubImageWidth ? CurrSubImageWidth - X : TileWidth;
 						Tasks.Add(TGraphTask<FProcessDecodedDataTask>::CreateTask().ConstructAndDispatchWhenReady([&UsableBufferQueue, &ProcessDecodedData, NumberOfColumn, NumberOfRow, BufferPtr, SampleIndex, TileX, TileY]()
 							{
 								ProcessDecodedData(NumberOfColumn, NumberOfRow, *(BufferPtr.Get()), SampleIndex, TileX, TileY);
@@ -1063,8 +1130,8 @@ namespace UE::ImageWrapper::Private
 			uint16 RowPerStrip = 0;
 			TIFFGetField(Tiff, TIFFTAG_ROWSPERSTRIP, &RowPerStrip);
 
-			const int32 NumberOfRow = RowPerStrip > Height || RowPerStrip == 0  ? Height : RowPerStrip;
-			if (StripByteSize < NumberOfRow * Width * sizeof(DataTypeSrc) * NumberOfChannelInReadArray / ReadWriteAdapter::ChannelPerReadIndex)
+			const int32 NumberOfRow = RowPerStrip > CurrSubImageHeight || RowPerStrip == 0  ? CurrSubImageHeight : RowPerStrip;
+			if (StripByteSize < NumberOfRow * CurrSubImageWidth * sizeof(DataTypeSrc) * NumberOfChannelInReadArray / ReadWriteAdapter::ChannelPerReadIndex)
 			{
 				// Lib tiff and this code is not able to deal with channel of different bit sizes
 				SetError(TEXT("Tiff strips are smaller then expected. This generaly due to channels of different size which is not supported by UE"));
@@ -1079,7 +1146,7 @@ namespace UE::ImageWrapper::Private
 			{
 				if (RowPerStrip != 0)
 				{ 
-					for (int32 Y = 0; Y < Height; Y += RowPerStrip)
+					for (int32 Y = 0; Y < CurrSubImageHeight; Y += RowPerStrip)
 					{
 						TSharedPtr<TArray64<uint8>, ESPMode::ThreadSafe> BufferPtr;
 						if (!UsableBufferQueue.Dequeue(BufferPtr))
@@ -1089,7 +1156,7 @@ namespace UE::ImageWrapper::Private
 						}
 
 						// The last strip might be smaller
-						int32 NumberOfRowRead = Y + RowPerStrip > Height ? Height - Y : RowPerStrip;
+						int32 NumberOfRowRead = Y + RowPerStrip > CurrSubImageHeight ? CurrSubImageHeight - Y : RowPerStrip;
 						if (TIFFReadEncodedStrip(Tiff, TIFFComputeStrip(Tiff, Y, SampleIndex), BufferPtr->GetData(), NumberOfRowRead * LineSizeScr) == -1)
 						{
 							SetError(ErrorMessage);
@@ -1100,7 +1167,7 @@ namespace UE::ImageWrapper::Private
 						Tasks.Add(TGraphTask<FProcessDecodedDataTask>::CreateTask().ConstructAndDispatchWhenReady([this, &ProcessDecodedData, &UsableBufferQueue, NumberOfRowRead, BufferPtr, Y, SampleIndex]()
 							{
 								// Consider the strip to be a tile that as the full with of the image.
-								ProcessDecodedData(Width, NumberOfRowRead, *(BufferPtr.Get()), SampleIndex, 0, Y);
+								ProcessDecodedData(CurrSubImageWidth, NumberOfRowRead, *(BufferPtr.Get()), SampleIndex, 0, Y);
 								UsableBufferQueue.Enqueue(BufferPtr);
 							}));
 					}
@@ -1123,7 +1190,7 @@ namespace UE::ImageWrapper::Private
 					Tasks.Add(TGraphTask<FProcessDecodedDataTask>::CreateTask().ConstructAndDispatchWhenReady([this, &UsableBufferQueue, &ProcessDecodedData, BufferPtr, SampleIndex]()
 						{
 							// Consider the image to be one big tile.
-							ProcessDecodedData(Width, Height, *(BufferPtr.Get()), SampleIndex, 0, 0);
+							ProcessDecodedData(CurrSubImageWidth, CurrSubImageHeight, *(BufferPtr.Get()), SampleIndex, 0, 0);
 							UsableBufferQueue.Enqueue(BufferPtr);
 						}));
 				}
@@ -1144,7 +1211,7 @@ namespace UE::ImageWrapper::Private
 			//todo: would be better to do this in the ProcessDecodedData tasks, so that we write the memory only once
 
 			//todo: ParallelFor over individual pixels is not great, better to do on rows; use ImageCore ImageParallelFor
-			ParallelFor(Width * Height, [NumOfChannelDest, &WriteArray](int32 Index)
+			ParallelFor(CurrSubImageWidth * CurrSubImageHeight, [NumOfChannelDest, &WriteArray](int32 Index)
 				{
 					const int64 WriteIndex = int64(NumOfChannelDest) * Index + NumOfChannelDest-1;
 					if constexpr (std::is_same_v<DataTypeDest, FFloat16> || TIsFloatingPoint<DataTypeDest>::Value)
diff --git a/Engine/Source/Runtime/ImageWrapper/Private/Formats/TiffImageWrapper.h b/Engine/Source/Runtime/ImageWrapper/Private/Formats/TiffImageWrapper.h
--- a/Engine/Source/Runtime/ImageWrapper/Private/Formats/TiffImageWrapper.h
+++ b/Engine/Source/Runtime/ImageWrapper/Private/Formats/TiffImageWrapper.h
@@ -20,20 +20,21 @@ namespace UE::ImageWrapper::Private
 		FTiffImageWrapper(FTiffImageWrapper&&) = default;
 		FTiffImageWrapper& operator=(FTiffImageWrapper&&) = default;
 
-
 		FTiffImageWrapper(const FTiffImageWrapper&) = delete;
 		FTiffImageWrapper& operator=(const FTiffImageWrapper&) = delete;
 
 		// FImageWrapperBase Interface
 		virtual void Compress(int32 Quality) override;
 		virtual void Uncompress(const ERGBFormat InFormat, int32 InBitDepth) override;
+		virtual void Uncompress(const ERGBFormat InFormat, int32 InBitDepth, FDecompressedImageOutput& OutDecompressedImage) override;
 		virtual bool SetCompressed(const void* InCompressedData, int64 InCompressedSize) override;
 		
 		virtual bool CanSetRawFormat(const ERGBFormat InFormat, const int32 InBitDepth) const override;
 		virtual ERawImageFormat::Type GetSupportedRawFormat(const ERawImageFormat::Type InFormat) const override;
 
 	private:
-
+		bool Uncompress_Internal(const ERGBFormat InFormat, int32 InBitDepth);
+		
 		void ReleaseTiffImage();
 
 		// Unpack the compressed data into the raw buffer. It will also add the alpha channel when needed.
@@ -63,6 +64,10 @@ namespace UE::ImageWrapper::Private
 		uint16 SamplesPerPixel = 0;
 		uint16 SampleFormat = 0;
 
+		int64 CurrSubImageWidth = 0;
+		int64 CurrSubImageHeight = 0;
+		TArray64<uint8> SubImageBuffer;
+
 		friend struct FTIFFReadMemoryFile;
 		
 		virtual void Reset() override
@@ -77,11 +82,9 @@ namespace UE::ImageWrapper::Private
 			BitsPerSample = 0;
 			SamplesPerPixel = 0;
 			SampleFormat = 0;
+
+			SubImageBuffer.Empty();
 		}
 	};
-
-
-
-
 }
 #endif // WITH_LIBTIFF
diff --git a/Engine/Source/Runtime/ImageWrapper/Private/ImageWrapperBase.cpp b/Engine/Source/Runtime/ImageWrapper/Private/ImageWrapperBase.cpp
--- a/Engine/Source/Runtime/ImageWrapper/Private/ImageWrapperBase.cpp
+++ b/Engine/Source/Runtime/ImageWrapper/Private/ImageWrapperBase.cpp
@@ -2,6 +2,7 @@
 
 #include "ImageWrapperBase.h"
 #include "ImageWrapperPrivate.h"
+#include "ImageWrapperOutputTypes.h"
 
 
 /* FImageWrapperBase structors
@@ -41,6 +42,13 @@ void FImageWrapperBase::SetError(const TCHAR* ErrorMessage)
 /* IImageWrapper structors
  *****************************************************************************/
 
+void FImageWrapperBase::Uncompress(const ERGBFormat InFormat, int32 InBitDepth, FDecompressedImageOutput& OutDecompressedImage)
+{
+	OutDecompressedImage = FDecompressedImageOutput();
+	Uncompress(InFormat, InBitDepth);
+	OutDecompressedImage.MipMapImage.AddMipImage(MoveTemp(RawData), GetWidth(), GetHeight());
+}
+
 TArray64<uint8> FImageWrapperBase::GetCompressed(int32 Quality)
 {
 	LastError.Empty();
@@ -55,19 +63,40 @@ bool FImageWrapperBase::GetRaw(const ERGBFormat InFormat, int32 InBitDepth, TArr
 	LastError.Empty();
 	Uncompress(InFormat, InBitDepth);
 
-	if ( ! LastError.IsEmpty())
+	if (!LastError.IsEmpty())
 	{
 		UE_LOG(LogImageWrapper, Warning, TEXT("ImageWrapper GetRaw failed: %s"), *LastError);
 		return false;
 	}
-	
-	if ( RawData.IsEmpty() )
+
+	if (RawData.IsEmpty())
 	{
 		return false;
 	}
 
 	OutRawData = MoveTemp(RawData);
-	
+
+	return true;
+}
+
+bool FImageWrapperBase::GetRaw(const ERGBFormat InFormat, int32 InBitDepth, FDecompressedImageOutput& OutDecompressedImage)
+{
+	LastError.Empty();
+
+	// Uncompress should add data directly to the mip map image through FMipMapImage::AddMipImage
+	Uncompress(InFormat, InBitDepth, OutDecompressedImage);
+
+	if (!LastError.IsEmpty())
+	{
+		UE_LOG(LogImageWrapper, Warning, TEXT("ImageWrapper GetRaw failed: %s"), *LastError);
+		return false;
+	}
+
+	if (OutDecompressedImage.MipMapImage.RawData.IsEmpty())
+	{
+		return false;
+	}
+
 	return true;
 }
 
@@ -391,7 +420,7 @@ bool FImageWrapperBase::GetImageViewOfSetRawForCompress(FImageView & OutImage) c
 bool IImageWrapper::GetRawImage(FImage & OutImage)
 {
 	TArray64<uint8> OutRawData;
-	if ( ! GetRaw(OutRawData) )
+	if (!GetRaw(OutRawData))
 	{
 		return false;
 	}
@@ -403,8 +432,8 @@ bool IImageWrapper::GetRawImage(FImage & OutImage)
 
 	bool bExactMatch;
 	ERawImageFormat::Type RawFormat = GetClosestRawImageFormat(&bExactMatch);
-	
-	if ( RawFormat == ERawImageFormat::Invalid )
+
+	if (RawFormat == ERawImageFormat::Invalid)
 	{
 		return false;
 	}
@@ -413,7 +442,7 @@ bool IImageWrapper::GetRawImage(FImage & OutImage)
 	//	just assume they are Default for now :
 	EGammaSpace GammaSpace = ERawImageFormat::GetDefaultGammaSpace(RawFormat);
 
-	if ( bExactMatch )
+	if (bExactMatch)
 	{
 		// no conversion required
 
@@ -426,34 +455,103 @@ bool IImageWrapper::GetRawImage(FImage & OutImage)
 	}
 	else
 	{
-		OutImage.Init( Width, Height, RawFormat, GammaSpace );
+		OutImage.Init(Width, Height, RawFormat, GammaSpace);
 
 		FImageView SrcImage = OutImage;
 		SrcImage.RawData = OutRawData.GetData();
 
-		switch(RGBFormat)
+		switch (RGBFormat)
 		{
 		case ERGBFormat::RGBA:
 		{
 			// RGBA8 -> BGRA8
-			check( BitDepth == 8 );
-			check( RawFormat == ERawImageFormat::BGRA8 );
-			FImageCore::CopyImageRGBABGRA(SrcImage, OutImage );
+			check(BitDepth == 8);
+			check(RawFormat == ERawImageFormat::BGRA8);
+			FImageCore::CopyImageRGBABGRA(SrcImage, OutImage);
 			break;
 		}
-			
+
 		case ERGBFormat::BGRA:
 		{
 			// BGRA16 -> RGBA16
-			check( BitDepth == 16 );
-			check( RawFormat == ERawImageFormat::RGBA16 );
-			FImageCore::CopyImageRGBABGRA(SrcImage, OutImage );
+			check(BitDepth == 16);
+			check(RawFormat == ERawImageFormat::RGBA16);
+			FImageCore::CopyImageRGBABGRA(SrcImage, OutImage);
 			break;
 		}
 
 		default:
 			check(0);
-			return false;			
+			return false;
+		}
+	}
+
+	return true;
+}
+
+bool IImageWrapper::GetRawImage(FDecompressedImageOutput& OutDecomressedImage)
+{
+	
+	if (!GetRaw(OutDecomressedImage))
+	{
+		return false;
+	}
+
+	int64 Width = GetWidth();
+	int64 Height = GetHeight();
+	ERGBFormat RGBFormat = GetFormat();
+	int32 BitDepth = GetBitDepth();
+
+	bool bExactMatch;
+	ERawImageFormat::Type RawFormat = GetClosestRawImageFormat(&bExactMatch);
+
+	if (RawFormat == ERawImageFormat::Invalid)
+	{
+		return false;
+	}
+
+	// ImageWrapper RGBFormat doesn't track if pixels are Gamma/sRGB or not
+	//	just assume they are Default for now :
+	EGammaSpace GammaSpace = ERawImageFormat::GetDefaultGammaSpace(RawFormat);
+
+	if (bExactMatch)
+	{
+		// no conversion required
+		OutDecomressedImage.MipMapImage.Format = RawFormat;
+		OutDecomressedImage.MipMapImage.GammaSpace = GammaSpace;
+	}
+	else
+	{
+		FMipMapImage& MipMapImage = OutDecomressedImage.MipMapImage;
+
+		for (int32 MipLevelIndex = 0; MipLevelIndex < MipMapImage.GetMipCount(); ++MipLevelIndex)
+		{
+			FImageView SrcImage = MipMapImage.GetMipImage(MipLevelIndex);
+
+			switch (RGBFormat)
+			{
+			case ERGBFormat::RGBA:
+			{
+				// RGBA8 -> BGRA8
+				check(BitDepth == 8);
+				check(RawFormat == ERawImageFormat::BGRA8);
+				FImageCore::TransposeImageRGBABGRA(SrcImage);
+				break;
+			}
+
+			case ERGBFormat::BGRA:
+			{
+				// BGRA16 -> RGBA16
+				check(BitDepth == 16);
+				check(RawFormat == ERawImageFormat::RGBA16);
+				FImageCore::TransposeImageRGBABGRA(SrcImage);
+				break;
+			}
+
+			default:
+				check(0);
+				return false;
+			}
 		}
 	}
 
diff --git a/Engine/Source/Runtime/ImageWrapper/Private/ImageWrapperBase.h b/Engine/Source/Runtime/ImageWrapper/Private/ImageWrapperBase.h
--- a/Engine/Source/Runtime/ImageWrapper/Private/ImageWrapperBase.h
+++ b/Engine/Source/Runtime/ImageWrapper/Private/ImageWrapperBase.h
@@ -77,11 +77,24 @@ public:
 	 * Function to uncompress our data 
 	 *
 	 * @param InFormat How we want to manipulate the RGB data
+	 * @param InBitDepth The bit depth per-channel of the image.
 	 * 
 	 * returns void. call SetError() in your implementation if you fail.
 	 */
 	virtual void Uncompress(const ERGBFormat InFormat, int32 InBitDepth) = 0;
 
+	/**
+	 * Function to uncompress our data and retrive meta info along with it.
+	 *
+	 * @param InFormat				How we want to manipulate the RGB data
+	 * @param InBitDepth			The bit depth per-channel of the image.
+	 * @param OutDecompressedImage	Holds decompressed image, its mip-maps if available along with meta 
+	 *								data that can be retrived during the decompression of the image.
+	 * 
+	 * returns void. call SetError() in your implementation if you fail.
+	 */
+	virtual void Uncompress(const ERGBFormat InFormat, int32 InBitDepth, FDecompressedImageOutput& OutDecompressedImage);
+
 public:
 
 	//~ IImageWrapper interface
@@ -105,6 +118,9 @@ public:
 
 	virtual bool GetRaw(const ERGBFormat InFormat, int32 InBitDepth, TArray64<uint8>& OutRawData) override;
 
+	// Use this version when you expect more information such as mipmaps, etc. 
+	virtual bool GetRaw(const ERGBFormat InFormat, int32 InBitDepth, FDecompressedImageOutput& OutDecompressedImage) override;
+
 	virtual int64 GetWidth() const override
 	{
 		return Width;
diff --git a/Engine/Source/Runtime/ImageWrapper/Private/ImageWrapperModule.cpp b/Engine/Source/Runtime/ImageWrapper/Private/ImageWrapperModule.cpp
--- a/Engine/Source/Runtime/ImageWrapper/Private/ImageWrapperModule.cpp
+++ b/Engine/Source/Runtime/ImageWrapper/Private/ImageWrapperModule.cpp
@@ -285,7 +285,8 @@ public:
 			return EImageFormat::HDR;
 		}
 		else if ( FCString::Stricmp(Name,TEXT("tiff")) == 0 ||
-			 FCString::Stricmp(Name,TEXT("tif")) == 0 )
+			 FCString::Stricmp(Name,TEXT("tif")) == 0 ||
+			FCString::Stricmp(Name, TEXT("tx")) == 0)
 		{
 			return EImageFormat::TIFF;
 		}
@@ -417,6 +418,37 @@ public:
 		return true;
 	}
 
+	virtual bool DecompressImage(const void* InCompressedData, int64 InCompressedSize, FDecompressedImageOutput& OutDecompressedImage) override
+	{
+		TRACE_CPUPROFILER_EVENT_SCOPE(ImageWrapper.Decompress);
+
+		OutDecompressedImage = FDecompressedImageOutput();
+
+		EImageFormat ImageFormat = DetectImageFormat(InCompressedData, InCompressedSize);
+		if (ImageFormat == EImageFormat::Invalid)
+		{
+			return false;
+		}
+
+		TSharedPtr<IImageWrapper> ImageWrapper = CreateImageWrapper(ImageFormat);
+		if (!ImageWrapper.IsValid())
+		{
+			return false;
+		}
+
+		if (!ImageWrapper->SetCompressed(InCompressedData, InCompressedSize))
+		{
+			return false;
+		}
+
+		if (!ImageWrapper->GetRawImage(OutDecompressedImage))
+		{
+			return false;
+		}
+
+		return true;
+	}
+
 	virtual bool CompressImage(TArray64<uint8> & OutData, EImageFormat ToFormat, const FImageView & InImage, int32 Quality) override
 	{
 		TRACE_CPUPROFILER_EVENT_SCOPE(ImageWrapper.Compress);
diff --git a/Engine/Source/Runtime/ImageWrapper/Public/Formats/HdrImageWrapper.h b/Engine/Source/Runtime/ImageWrapper/Public/Formats/HdrImageWrapper.h
--- a/Engine/Source/Runtime/ImageWrapper/Public/Formats/HdrImageWrapper.h
+++ b/Engine/Source/Runtime/ImageWrapper/Public/Formats/HdrImageWrapper.h
@@ -26,6 +26,7 @@ public:
 	IMAGEWRAPPER_API virtual bool SetRaw(const void* InRawData, int64 InRawSize, const int32 InWidth, const int32 InHeight, const ERGBFormat InFormat, const int32 InBitDepth, const int32 InBytesPerRow) override;
 	IMAGEWRAPPER_API virtual TArray64<uint8> GetCompressed(int32 Quality = 0) override;
 	IMAGEWRAPPER_API virtual bool GetRaw(const ERGBFormat InFormat, int32 InBitDepth, TArray64<uint8>& OutRawData) override;
+	IMAGEWRAPPER_API virtual bool GetRaw(const ERGBFormat InFormat, int32 InBitDepth, FDecompressedImageOutput& OutDecompressedImage) override;
 	
 	IMAGEWRAPPER_API virtual bool CanSetRawFormat(const ERGBFormat InFormat, const int32 InBitDepth) const override;
 	IMAGEWRAPPER_API virtual ERawImageFormat::Type GetSupportedRawFormat(const ERawImageFormat::Type InFormat) const override;
diff --git a/Engine/Source/Runtime/ImageWrapper/Public/IImageWrapper.h b/Engine/Source/Runtime/ImageWrapper/Public/IImageWrapper.h
--- a/Engine/Source/Runtime/ImageWrapper/Public/IImageWrapper.h
+++ b/Engine/Source/Runtime/ImageWrapper/Public/IImageWrapper.h
@@ -7,6 +7,7 @@
 #include "CoreTypes.h"
 #include "Templates/SharedPointer.h"
 #include "ImageCore.h"
+#include "ImageWrapperOutputTypes.h"
 
 /**
 
@@ -169,10 +170,11 @@ public:
 	 */
 	virtual bool SetRaw(const void* InRawData, int64 InRawSize, const int32 InWidth, const int32 InHeight, const ERGBFormat InFormat, const int32 InBitDepth, const int32 InBytesPerRow = 0) = 0;
 
-	/* CanSetRawFormat returns true if SetRaw will accept this format */
+	/** CanSetRawFormat returns true if SetRaw will accept this format */
 	virtual bool CanSetRawFormat(const ERGBFormat InFormat, const int32 InBitDepth) const = 0;
 
-	/* returns InFormat if supported, else maps to something supported
+	/** 
+	 * returns InFormat if supported, else maps to something supported
 	 * the returned format will pass CanSetRawFormat()
 	 */
 	virtual ERawImageFormat::Type GetSupportedRawFormat(const ERawImageFormat::Type InFormat) const = 0;
@@ -200,14 +202,14 @@ public:
 	}
 
 	/**
-	* GetRaw after SetCompressed
-	* fills the raw data in the native format/depth contained in the file
-	* Do not use the GetRaw() variants that take format/depth arguments.
-	* 
-	* @param OutRawData    Filled with raw image data.
-	*
-	* This GetRaw call replaces the variants with format/depth arguments, but prefer GetRawImage instead.
-	*/
+	 * GetRaw after SetCompressed
+	 * fills the raw data in the native format/depth contained in the file
+	 * Do not use the GetRaw() variants that take format/depth arguments.
+	 * 
+	 * @param OutRawData    Filled with raw image data.
+	 *
+	 * This GetRaw call replaces the variants with format/depth arguments, but prefer GetRawImage instead.
+	 */
 	bool GetRaw(TArray64<uint8>& OutRawData)
 	{
 		ERGBFormat Format = GetFormat();
@@ -221,19 +223,58 @@ public:
 
 		return GetRaw(Format,BitDepth,OutRawData);
 	}
+
+	/**
+	 * GetRaw after SetCompressed
+	 * fills the raw data in the native format/depth contained in the file along with meta info
+	 * Do not use the GetRaw() variants that take format/depth arguments.
+	 *
+	 * @param OutRawData				Filled with raw image data.
+	 * @param OutDecompressedImage		Filled with mip images and other metadata.
+	 *
+	 * This GetRaw call replaces the variants with format/depth arguments, but prefer GetRawImage instead.
+	 */
+	bool GetRaw(FDecompressedImageOutput& OutDecompressedImage)
+	{
+		ERGBFormat Format = GetFormat();
+		int32 BitDepth = GetBitDepth();
+
+		// Format and BitDepth should have been set by SetCompressed
+		if (Format == ERGBFormat::Invalid || BitDepth == 0)
+		{
+			return false;
+		}
+
+		return GetRaw(Format, BitDepth, OutDecompressedImage);
+	}
 	
-	/* Decode the image file data from SetCompressed() into an FImage
-	* OutImage is allocated and attributes are filled
-	*	 any previous passed-in contents of OutImage are destroyed
-	*	OutImage.Format is ignored, a new format is set from the loaded image
-	*
-	* @param OutImage	 Filled with the image
-	*
-	* This is the recommended API to get the raw image data from the imagewrapper.
-	* Prefer this instead of any of the GetRaw() calls.
-	*/
+	/**
+	 * Decode the image file data from SetCompressed() into an FImage
+	 * OutImage is allocated and attributes are filled
+	 * Any previous passed-in contents of OutImage are destroyed
+	 * OutImage.Format is ignored, a new format is set from the loaded image
+	 *
+	 * @param OutImage	 Filled with the image
+	 *
+	 * This is the recommended API to get the raw image data from the imagewrapper.
+	 * Prefer this instead of any of the GetRaw() calls.
+	 */
 	bool GetRawImage(FImage & OutImage);
 
+	/**
+	 * Decode the image file data from SetCompressed() into an FImage
+	 * OutImage is allocated and attributes are filled.
+	 * Any previous passed-in contents of OutImage are destroyed.
+	 * OutImage.Format is ignored, a new format is set from the loaded image
+	 * MetaInfo is filled if applicable
+	 * 
+	 * @param OutDecompressedImage	Image and meta data along with Mip Map Images
+	 * 
+	 * This is the recommended API to get the raw image data from the imagewrapper.
+	 * Prefer this instead of any of the GetRaw() calls.
+	 */
+	bool GetRawImage(FDecompressedImageOutput& OutDecompressedImage);
+
 	
 	/**  
 	 * Gets the raw data.
@@ -248,7 +289,22 @@ public:
 	 * DEPRECATED , use GetRaw() with 1 argument or GetRawImage()
 	 */
 	virtual bool GetRaw(const ERGBFormat InFormat, int32 InBitDepth, TArray64<uint8>& OutRawData) = 0;
-	
+
+	/**
+	 * Gets the raw data and meta info
+	 * (Note: It may consume the data set in the SetRaw function if it was set before)
+	 *
+	 * @param InFormat How we want to manipulate the RGB data.
+	 * @param InBitDepth The output bit-depth per channel, normally 8.
+	 * @param OutRawData Will contain the uncompressed raw data.
+	 * @param OutDecompressedImage Will contain the mip images if available and other metadata.
+	 * @return true on success, false otherwise.
+	 *
+	 * this is often broken, should only be used with InFormat == GetFormat()
+	 * DEPRECATED , use GetRaw() with 2 argument or GetRawImage() meta info overload
+	 */
+	virtual bool GetRaw(const ERGBFormat InFormat, int32 InBitDepth, FDecompressedImageOutput& OutDecompressedImage) = 0;
+		
 	/**
 	 * Gets the raw data in a TArray. Only use this if you're certain that the image is less than 2 GB in size.
 	 * Prefer using the overload which takes a TArray64 in general.
diff --git a/Engine/Source/Runtime/ImageWrapper/Public/IImageWrapperModule.h b/Engine/Source/Runtime/ImageWrapper/Public/IImageWrapperModule.h
--- a/Engine/Source/Runtime/ImageWrapper/Public/IImageWrapperModule.h
+++ b/Engine/Source/Runtime/ImageWrapper/Public/IImageWrapperModule.h
@@ -7,6 +7,9 @@
 #include "Templates/SharedPointer.h"
 #include "ImageCore.h"
 #include "IImageWrapper.h"
+#include "ImageWrapperOutputTypes.h"
+
+
 
 
 /**
@@ -62,6 +65,16 @@ public:
 	 */
 	virtual bool DecompressImage(const void* InCompressedData, int64 InCompressedSize, FImage & OutImage) = 0;
 
+	/* Read an image from file format encoded data.
+	 * ImageWrapper calls this a "decompress"
+	 * OutImage is allocated and filled, any existing contents are discarded
+	 *
+	 * @param InCompressedData					Image format encoded bytes to read; format is automatically deduced from the content
+	 * @param InCompressedSize					Size of InCompressedData in bytes
+	 * @param OutDecompressedImageData			Filled with Image along with other meta data and mip images that is read.  Allocated.
+	 */
+	virtual bool DecompressImage(const void* InCompressedData, int64 InCompressedSize, FDecompressedImageOutput& OutDecompressedImage) = 0;
+
 	/**  
 	 * Create an IImageWrapper helper of a specific type 
 	 *
diff --git a/Engine/Source/Runtime/ImageWrapper/Public/ImageWrapperOutputTypes.h b/Engine/Source/Runtime/ImageWrapper/Public/ImageWrapperOutputTypes.h
new file mode 100644
--- /dev/null
+++ b/Engine/Source/Runtime/ImageWrapper/Public/ImageWrapperOutputTypes.h
@@ -0,0 +1,13 @@
+// Copyright Epic Games, Inc. All Rights Reserved.
+
+#pragma once
+#include "CoreTypes.h"
+#include "ImageCore.h"
+
+/**
+ * Use this struct to get mip images and other metadata about the image.
+ */
+struct FDecompressedImageOutput
+{
+	FMipMapImage MipMapImage;
+};
\ No newline at end of file
diff --git a/Engine/Source/Runtime/TextureUtilitiesCommon/Private/TextureImportUtils.cpp b/Engine/Source/Runtime/TextureUtilitiesCommon/Private/TextureImportUtils.cpp
--- a/Engine/Source/Runtime/TextureUtilitiesCommon/Private/TextureImportUtils.cpp
+++ b/Engine/Source/Runtime/TextureUtilitiesCommon/Private/TextureImportUtils.cpp
@@ -12,12 +12,8 @@ namespace UE
 {
 	namespace TextureUtilitiesCommon
 	{
-		/**
-		 * Detect the existence of gray scale image in some formats and convert those to a gray scale equivalent image
-		 * 
-		 * @return true if the image was converted
-		 */
-		bool AutoDetectAndChangeGrayScale(FImage& Image)
+		template<typename TImageClass>
+		bool AutoDetectAndChangeGrayScale(TImageClass& Image)
 		{
 			if (Image.Format != ERawImageFormat::BGRA8)
 			{
@@ -35,7 +31,7 @@ namespace UE
 					Colors[i].R != Colors[i].B ||
 					Colors[i].G != Colors[i].B)
 				{
-					return false ;
+					return false;
 				}
 			}
 
@@ -45,6 +41,9 @@ namespace UE
 			return true;
 		}
 
+		template TEXTUREUTILITIESCOMMON_API bool AutoDetectAndChangeGrayScale(FImage& Image);
+		template TEXTUREUTILITIESCOMMON_API bool AutoDetectAndChangeGrayScale(FMipMapImage& Image);
+
 		/**
 		 * This fills any pixels of a texture with have an alpha value of zero and RGB=white,
 		 * with an RGB from the nearest neighboring pixel which has non-zero alpha.
diff --git a/Engine/Source/Runtime/TextureUtilitiesCommon/Public/TextureImportUtils.h b/Engine/Source/Runtime/TextureUtilitiesCommon/Public/TextureImportUtils.h
--- a/Engine/Source/Runtime/TextureUtilitiesCommon/Public/TextureImportUtils.h
+++ b/Engine/Source/Runtime/TextureUtilitiesCommon/Public/TextureImportUtils.h
@@ -5,16 +5,15 @@
 #include "CoreTypes.h"
 #include "Engine/TextureDefines.h"
 
-struct FImage;
-
 namespace UE::TextureUtilitiesCommon
 {
 	/**
 	 * Detect the existence of gray scale image in some formats and convert those to a gray scale equivalent image
-	 * 
+	 *
 	 * @return true if the image was converted
 	 */
-	TEXTUREUTILITIESCOMMON_API bool AutoDetectAndChangeGrayScale(FImage& Image);
+	template<typename TImageClass>
+	TEXTUREUTILITIESCOMMON_API bool AutoDetectAndChangeGrayScale(TImageClass& Image);
 
 	/**
 	 * For PNG texture importing, this ensures that any pixels with an alpha value of zero have an RGB
