[BUGFIX] TIFF Image Wrapper
- Error out if the strips or tiles exceed the max supported index for tasks array
- Fixed a few other error messages.

#rnx

diff --git a/Engine/Source/Runtime/ImageWrapper/Private/Formats/TiffImageWrapper.cpp b/Engine/Source/Runtime/ImageWrapper/Private/Formats/TiffImageWrapper.cpp
--- a/Engine/Source/Runtime/ImageWrapper/Private/Formats/TiffImageWrapper.cpp
+++ b/Engine/Source/Runtime/ImageWrapper/Private/Formats/TiffImageWrapper.cpp
@@ -1088,11 +1088,19 @@ namespace UE::ImageWrapper::Private
 			if (TileSize < TileHeight * TileWidth * NumberOfChannelInReadArray * sizeof(DataTypeSrc) / ReadWriteAdapter::ChannelPerReadIndex)
 			{
 				// Lib tiff and this code is not able to deal with channel of different bit sizes
-				SetError(TEXT("Tiff tiles are smaller then expected. This generaly due to channels of different size which is not supported by UE"));
+				SetError(TEXT("Tiff tiles are smaller than expected. This is generally due to channels of different size which is not supported by UE"));
 				return false;
 			}
 
-			Tasks.Reserve(TIFFNumberOfTiles(Tiff));
+			const uint64 NumTiles = TIFFNumberOfTiles(Tiff);
+			if (NumTiles > (uint64)INT32_MAX)
+			{
+				const FString TileErrorMessage = FString::Printf(TEXT("UE can support at most %u tiles at the moment. The file contains %u tiles."), INT32_MAX, NumTiles);
+				SetError(*TileErrorMessage);
+				return false;
+			}
+
+			Tasks.Reserve(NumTiles);
 
 			for (uint8 SampleIndex = 0; SampleIndex < NumOfPlanes; SampleIndex++)
 			{ 
@@ -1111,7 +1119,7 @@ namespace UE::ImageWrapper::Private
 
 						if (TIFFReadTile(Tiff, BufferPtr->GetData(), X, Y, 0, SampleIndex) < 0)
 						{
-							SetError(TEXT("Couldn't open Tiff tile. This is generally caused by a compression format that UE don't support."));
+							SetError(TEXT("Couldn't open Tiff tile. This is generally caused by a compression format that UE doesn't support."));
 							return false;
 						}
 
@@ -1134,14 +1142,22 @@ namespace UE::ImageWrapper::Private
 			if (StripByteSize < NumberOfRow * CurrSubImageWidth * sizeof(DataTypeSrc) * NumberOfChannelInReadArray / ReadWriteAdapter::ChannelPerReadIndex)
 			{
 				// Lib tiff and this code is not able to deal with channel of different bit sizes
-				SetError(TEXT("Tiff strips are smaller then expected. This generaly due to channels of different size which is not supported by UE"));
+				SetError(TEXT("Tiff strips are smaller than expected. This is generally due to channels of different size which is not supported by UE"));
 				return false;
 			}
 
-			const TCHAR* ErrorMessage = TEXT("Couldn't open Tiff strip. This is generally caused by a compression format that UE don't support.");
+			const uint64 NumStrips = TIFFNumberOfStrips(Tiff);
+			if (NumStrips > (uint64)INT32_MAX)
+			{
+				const FString StripErrorMessage = FString::Printf(TEXT("UE can support at most %u strips at the moment. The file contains %u strips."), INT32_MAX, NumStrips);
+				SetError(*StripErrorMessage);
+				return false;
+			}
 
-			Tasks.Reserve(TIFFNumberOfStrips(Tiff));
+			const TCHAR* ErrorMessage = TEXT("Couldn't open Tiff strip. This is generally caused by a compression format that UE doesn't support.");
 
+			Tasks.Reserve(NumStrips);
+			
 			for (uint8 SampleIndex = 0; SampleIndex < NumOfPlanes; SampleIndex++)
 			{
 				if (RowPerStrip != 0)
@@ -1231,5 +1247,4 @@ namespace UE::ImageWrapper::Private
 	}
 
 }
-
-#endif // WITH_LIBTIFF
+#endif // WITH_LIBTIFF
\ No newline at end of file
