ImageWrapper Jpeg: Add insights trace to jpeg-turbo compression and decompression.

Helps analyze performance bottlenecks.

#rnx

diff --git a/Engine/Source/Runtime/ImageWrapper/Private/Formats/JpegImageWrapper.cpp b/Engine/Source/Runtime/ImageWrapper/Private/Formats/JpegImageWrapper.cpp
--- a/Engine/Source/Runtime/ImageWrapper/Private/Formats/JpegImageWrapper.cpp
+++ b/Engine/Source/Runtime/ImageWrapper/Private/Formats/JpegImageWrapper.cpp
@@ -363,6 +363,8 @@ bool FJpegImageWrapper::SetCompressedTurbo(const void* InCompressedData, int64 I
 
 void FJpegImageWrapper::CompressTurbo(int32 Quality)
 {
+	TRACE_CPUPROFILER_EVENT_SCOPE(FJpegImageWrapper::CompressTurbo);
+
 	if (CompressedData.Num() == 0)
 	{
 		JPEG_SCOPE_CRITSEC();
@@ -409,6 +411,8 @@ void FJpegImageWrapper::CompressTurbo(int32 Quality)
 
 void FJpegImageWrapper::UncompressTurbo(const ERGBFormat InFormat, int32 InBitDepth)
 {
+	TRACE_CPUPROFILER_EVENT_SCOPE(FJpegImageWrapper::UncompressTurbo);
+
 	// Ensure we haven't already uncompressed the file.
 	if (RawData.Num() != 0)
 	{
