[ImageWrapper] Made sure the CompressedBuffer for LibJPG has a minimum size so compression doesn't fail when compressed size > sourcesize

diff --git a/Engine/Source/Runtime/ImageWrapper/Private/Formats/JpegImageWrapper.cpp b/Engine/Source/Runtime/ImageWrapper/Private/Formats/JpegImageWrapper.cpp
--- a/Engine/Source/Runtime/ImageWrapper/Private/Formats/JpegImageWrapper.cpp
+++ b/Engine/Source/Runtime/ImageWrapper/Private/Formats/JpegImageWrapper.cpp
@@ -239,12 +239,15 @@ void FJpegImageWrapper::Compress(int32 Quality)
 			}
 		}
 
-		CompressedData.Reset(RawData.Num());
-		CompressedData.AddUninitialized(RawData.Num());
 
 		// Note: OutBufferSize intentionally uses int64_t type as that's what jpge::compress_image_to_jpeg_file_in_memory expects.
 		// UE int64 type is not compatible with int64_t on all compilers (int64_t may be `long`, while int64 is `long long`).
-		int64_t OutBufferSize = CompressedData.Num();
+		// jpge::compress_image_to_jpeg_file_in_memory recommends at least 1kb buffer size since small images
+		// have some overhead and certain configurations can expand in size (e.g. high quality)
+		int64_t OutBufferSize = std::max<int64>(RawData.Num(), 1024);
+		
+		CompressedData.Reset(OutBufferSize);
+		CompressedData.AddUninitialized(OutBufferSize);
 
 		jpge::params Parameters;
 		Parameters.m_quality = Quality;
