ImageWrapperModule.h nop comments

diff --git a/Engine/Source/Runtime/ImageWrapper/Public/IImageWrapperModule.h b/Engine/Source/Runtime/ImageWrapper/Public/IImageWrapperModule.h
--- a/Engine/Source/Runtime/ImageWrapper/Public/IImageWrapperModule.h
+++ b/Engine/Source/Runtime/ImageWrapper/Public/IImageWrapperModule.h
@@ -14,14 +14,25 @@
 
 /**
  * Interface for image wrapper modules.
+ *
+ * DO NOT WRITE CODE THAT USES IMAPE WRAPPERS DIRECTLY.
  * 
  * If you have "Engine" module, then just use FImageUtils::CompressImage/DecompressImage
- *
  * The easiest way to load/save an image is FImageUtils::LoadImage/SaveImageByExtension
  *
+ * If you don't have "Engine" module, then use IImageWrapperModule::CompressImage/DecompressImage
+ * Do not manually create individual image wrappers.
+ *
  * NOTE: Please prefer to work with images through ImageCore FImage/FImageView 
  * (not TextureSource/TextureSourceFormat or raw arrays of bytes)
+ *  again, do not manually use ERGBFormat and arrays of bytes to work with ImageWrappers
+ *  instead use the FImage-oriented APIs which will handle all format conversions for you.
  * 
+ * You do NOT need to convert formats before saving an image, the FImage APIs will do it for you.
+ *
+ * If you got an image from a render target, it may have a zero alpha channel.  You might want to do
+ *   FImageCore::SetAlphaOpaque before saving.
+ *
  * 
  * note on SRGB/Gamma handling :
  * it is assumed that non-U8 data is always Linear
