[BUGFIX] Tiff Image Wrapper
- Tiff Wrapper handles ASCII Tags properly to avoid treating inaccessible memory as string data.
- Tiff Wrapper skips anonymous tags for now.

#rnx

diff --git a/Engine/Source/Runtime/ImageWrapper/Private/Formats/TiffImageWrapper.cpp b/Engine/Source/Runtime/ImageWrapper/Private/Formats/TiffImageWrapper.cpp
--- a/Engine/Source/Runtime/ImageWrapper/Private/Formats/TiffImageWrapper.cpp
+++ b/Engine/Source/Runtime/ImageWrapper/Private/Formats/TiffImageWrapper.cpp
@@ -534,14 +534,45 @@ namespace UE::ImageWrapper::Private
 				continue;
 			}
 
-			if (TIFFFieldDataType(TiffField) == TIFF_ASCII)
+			TIFFDataType TiffDataType = TIFFFieldDataType(TiffField);
+
+			if (TiffDataType != TIFF_ASCII)
+			{
+				continue;
+			}
+
+			const int32 PassCount = TIFFFieldPassCount(TiffField);
+			const int32 ReadCount = TIFFFieldReadCount(TiffField);
+
+			char* FieldCString = nullptr;
+			bool bCanProcessField = false;
+			
+			if (PassCount > 0 && ReadCount == TIFF_VARIABLE)
+			{
+				uint16 SmallValueCount = 0;
+				bCanProcessField = TIFFGetField(Tiff, TagListEntry, &SmallValueCount, &FieldCString) == 1;
+			}
+			else if (PassCount > 0 && ReadCount == TIFF_VARIABLE2)
 			{
-				void* TagValueRawData = nullptr;
-				if (TIFFGetField(Tiff, TagListEntry, &TagValueRawData) != 1)
+				uint32 FieldValueCount = 0;
+				bCanProcessField = TIFFGetField(Tiff, TagListEntry, &FieldValueCount, &FieldCString) == 1;
+			}
+			else if (ReadCount > 0)
+			{
+				bCanProcessField = TIFFGetField(Tiff, TagListEntry, &FieldCString) == 1;
+			}
+			else
+			{
+				bCanProcessField = TIFFGetField(Tiff, TagListEntry, &FieldCString) == 1;
+			}
+			
+			if (bCanProcessField && FieldCString && *FieldCString)
+			{
+				// Skip anonymous fields for now.
+				if (const char* FieldName = TIFFFieldName(TiffField))
 				{
-					continue;
+					OutTagMetaData.Emplace(FieldName, (const char*)FieldCString);
 				}
-				OutTagMetaData.Emplace(TIFFFieldName(TiffField), (const char*)TagValueRawData);
 			}
 		}
 
@@ -1161,6 +1192,9 @@ namespace UE::ImageWrapper::Private
 					}
 				}
 			}
+
+			/* Temp code to debug https://jira.it.epicgames.com/browse/UE-238630 */
+			UE_LOG(LogImageWrapper, Display, TEXT("Originally Reserved on NumTiles: %llu, Current Task Count: %llu"), NumTiles, (uint64)Tasks.Num());
 		}
 		else
 		{
@@ -1240,6 +1274,8 @@ namespace UE::ImageWrapper::Private
 						}));
 				}
 			}
+			/* Temp code to debug https://jira.it.epicgames.com/browse/UE-238630 */
+			UE_LOG(LogImageWrapper, Display, TEXT("Originally Reserved on NumStrips: %llu, Current Task Count: %llu"), NumStrips, (uint64)Tasks.Num());
 		}
 
 		// Write the alpha
