[Interchange] TIFF Image Wrapper - Remove Debug Logs

#rnx

diff --git a/Engine/Source/Runtime/ImageWrapper/Private/Formats/TiffImageWrapper.cpp b/Engine/Source/Runtime/ImageWrapper/Private/Formats/TiffImageWrapper.cpp
--- a/Engine/Source/Runtime/ImageWrapper/Private/Formats/TiffImageWrapper.cpp
+++ b/Engine/Source/Runtime/ImageWrapper/Private/Formats/TiffImageWrapper.cpp
@@ -1192,9 +1192,6 @@ namespace UE::ImageWrapper::Private
 					}
 				}
 			}
-
-			/* Temp code to debug https://jira.it.epicgames.com/browse/UE-238630 */
-			UE_LOG(LogImageWrapper, Display, TEXT("Originally Reserved on NumTiles: %llu, Current Task Count: %llu"), NumTiles, (uint64)Tasks.Num());
 		}
 		else
 		{
@@ -1274,8 +1271,6 @@ namespace UE::ImageWrapper::Private
 						}));
 				}
 			}
-			/* Temp code to debug https://jira.it.epicgames.com/browse/UE-238630 */
-			UE_LOG(LogImageWrapper, Display, TEXT("Originally Reserved on NumStrips: %llu, Current Task Count: %llu"), NumStrips, (uint64)Tasks.Num());
 		}
 
 		// Write the alpha
