[ImageWrapper]
* Added workaround for certain versions of png lib issue

diff --git a/Engine/Source/Runtime/ImageWrapper/Private/ImageWrapperModule.cpp b/Engine/Source/Runtime/ImageWrapper/Private/ImageWrapperModule.cpp
--- a/Engine/Source/Runtime/ImageWrapper/Private/ImageWrapperModule.cpp
+++ b/Engine/Source/Runtime/ImageWrapper/Private/ImageWrapperModule.cpp
@@ -531,3 +531,8 @@ public:
 
 
 IMPLEMENT_MODULE(FImageWrapperModule, ImageWrapper);
+
+
+#if PLATFORM_WINDOWS && !UE_BUILD_SHIPPING
+extern "C" DLLEXPORT void png_write_sig2() {} // To work around png lib issue. Ask honk for details
+#endif
