Fix loading png images when 8/16-bit depth conversions are enabled

diff --git a/Engine/Source/Runtime/ImageWrapper/Private/Formats/PngImageWrapper.cpp b/Engine/Source/Runtime/ImageWrapper/Private/Formats/PngImageWrapper.cpp
--- a/Engine/Source/Runtime/ImageWrapper/Private/Formats/PngImageWrapper.cpp
+++ b/Engine/Source/Runtime/ImageWrapper/Private/Formats/PngImageWrapper.cpp
@@ -483,24 +483,13 @@ void FPngImageWrapper::UncompressPNGData(const ERGBFormat InFormat, const int32
 			// Reduce 16-bit to 8-bit if requested
 			if (BitDepth == 16 && InBitDepth == 8)
 			{
-#if PNG_LIBPNG_VER >= 10504
-				check(0); // Needs testing
 				Transform |= PNG_TRANSFORM_SCALE_16;
-#else
-				Transform |= PNG_TRANSFORM_STRIP_16;
-#endif
 			}
 
 			// Increase 8-bit to 16-bit if requested
 			if (BitDepth <= 8 && InBitDepth == 16)
 			{
-#if PNG_LIBPNG_VER >= 10504
-				check(0); // Needs testing
 				Transform |= PNG_TRANSFORM_EXPAND_16;
-#else
-				// Expanding 8-bit images to 16-bit via transform needs a libpng update
-				check(0);
-#endif
 			}
 
 			png_read_png(png_ptr, info_ptr, Transform, NULL);
