PngImageWrapper downgrade verbosity of some spammy warnings

(Introduced in libpng 1.6)

diff --git a/Engine/Source/Runtime/ImageWrapper/Private/Formats/PngImageWrapper.cpp b/Engine/Source/Runtime/ImageWrapper/Private/Formats/PngImageWrapper.cpp
--- a/Engine/Source/Runtime/ImageWrapper/Private/Formats/PngImageWrapper.cpp
+++ b/Engine/Source/Runtime/ImageWrapper/Private/Formats/PngImageWrapper.cpp
@@ -693,7 +693,17 @@ void FPngImageWrapper::user_warning_fn(png_structp png_ptr, png_const_charp warn
 {
 	FPngImageWrapper* ctx = (FPngImageWrapper*)png_get_error_ptr(png_ptr);
 	const TCHAR* LocalDebugImageName = ctx != nullptr ? ctx->DebugImageName : nullptr;
-	UE_LOG(LogImageWrapper, Warning, TEXT("PNG Warning(%s) %s"), LocalDebugImageName != nullptr ? LocalDebugImageName : TEXT(""), ANSI_TO_TCHAR(warning_msg));
+
+	// libPNG 1.6+ has some very spammy warnings about incorrect color profile metadata that is unfortunately
+	// extremely common in the wild, so manually detect and downgrade the verbosity of that
+	if (FPlatformString::Strncmp(warning_msg, "iCCP:", 5) == 0)
+	{
+		UE_LOG(LogImageWrapper, Verbose, TEXT("PNG Warning(%s) %s"), LocalDebugImageName != nullptr ? LocalDebugImageName : TEXT(""), ANSI_TO_TCHAR(warning_msg));
+	}
+	else
+	{
+		UE_LOG(LogImageWrapper, Warning, TEXT("PNG Warning(%s) %s"), LocalDebugImageName != nullptr ? LocalDebugImageName : TEXT(""), ANSI_TO_TCHAR(warning_msg));
+	}
 }
 
 void* FPngImageWrapper::user_malloc(png_structp /*png_ptr*/, png_size_t size)
