Add support for writing to the PNG text header

Adding a couple functions to IImageWrapper for getting and setting metadata from the image's metadata block (if supported).  Currently, this is just implemented for PNG files, and the other concrete classes just return false / not supported.  For PNGs, however, you may store any number of key/value pairs in the text section of the PNG's header, and have them restored when the image is later loaded.

#tests Saved and loaded PNGs both with and without metadata tags, and tested for leaks

diff --git a/Engine/Source/Runtime/ImageWrapper/Private/Formats/HdrImageWrapper.cpp b/Engine/Source/Runtime/ImageWrapper/Private/Formats/HdrImageWrapper.cpp
--- a/Engine/Source/Runtime/ImageWrapper/Private/Formats/HdrImageWrapper.cpp
+++ b/Engine/Source/Runtime/ImageWrapper/Private/Formats/HdrImageWrapper.cpp
@@ -652,4 +652,18 @@ void FHdrImageWrapper::FreeCompressedData()
 	CompressedDataHolder.Empty();
 }
 
+bool FHdrImageWrapper::SupportsMetadata() const
+{
+	return false;
+}
+
+void FHdrImageWrapper::AddMetadata(const FString&, const FString&)
+{
+}
+
+bool FHdrImageWrapper::TryGetMetadata(const FString&, FString&) const
+{
+	return false;
+}
+
 #undef LOCTEXT_NAMESPACE
diff --git a/Engine/Source/Runtime/ImageWrapper/Private/Formats/PngImageWrapper.cpp b/Engine/Source/Runtime/ImageWrapper/Private/Formats/PngImageWrapper.cpp
--- a/Engine/Source/Runtime/ImageWrapper/Private/Formats/PngImageWrapper.cpp
+++ b/Engine/Source/Runtime/ImageWrapper/Private/Formats/PngImageWrapper.cpp
@@ -3,6 +3,7 @@
 #include "Formats/PngImageWrapper.h"
 #include "ImageWrapperPrivate.h"
 
+#include "Containers/Utf8String.h"
 #include "Math/GuardedInt.h"
 #include "Misc/ScopeLock.h"
 #include "ImageCoreUtils.h"
@@ -90,6 +91,7 @@ public:
 		{
 			png_free(*PNGWriteStruct, PNGRowPointers);
 		}
+
 		png_destroy_write_struct(PNGWriteStruct, info_ptr);
 	}
 
@@ -98,8 +100,10 @@ public:
 		PNGRowPointers = InRowPointers;
 	}
 
-private:
+	TArray<png_text> MetadataKeyValuePairs;
+	TArray<FUtf8String> MetadataStrings;
 
+private:
 	png_structp* PNGWriteStruct;
 	png_infop* info_ptr;
 	png_bytep* PNGRowPointers;
@@ -175,6 +179,7 @@ void FPngImageWrapper::Compress(int32 Quality)
 		check(info_ptr);
 
 		png_bytep* row_pointers = (png_bytep*) png_malloc( png_ptr, Height*sizeof(png_bytep) );
+		
 		PNGWriteGuard PNGGuard(&png_ptr, &info_ptr);
 		PNGGuard.SetRowPointers( row_pointers );
 
@@ -220,6 +225,25 @@ void FPngImageWrapper::Compress(int32 Quality)
 				}
 			}
 
+			PNGGuard.MetadataKeyValuePairs.Reserve(TextBlocks.Num());
+			PNGGuard.MetadataStrings.Reserve(TextBlocks.Num() * 2);
+
+			for (auto& Block : TextBlocks)
+			{
+				png_text info {};
+				info.compression = PNG_TEXT_COMPRESSION_NONE;
+
+				PNGGuard.MetadataStrings.Emplace((*Block.Key));
+				info.key = (char*)*PNGGuard.MetadataStrings.Last();
+
+				PNGGuard.MetadataStrings.Emplace((*Block.Value));
+				info.text = (char*)*PNGGuard.MetadataStrings.Last();
+
+				PNGGuard.MetadataKeyValuePairs.Add(info);
+			}
+
+			png_set_text(png_ptr, info_ptr, PNGGuard.MetadataKeyValuePairs.GetData(), PNGGuard.MetadataKeyValuePairs.Num());
+
 			png_set_compression_level(png_ptr, ZlibLevel);
 			png_set_IHDR(png_ptr, info_ptr, Width, Height, BitDepth, (Format == ERGBFormat::Gray) ? PNG_COLOR_TYPE_GRAY : PNG_COLOR_TYPE_RGBA, PNG_INTERLACE_NONE, PNG_COMPRESSION_TYPE_DEFAULT, PNG_FILTER_TYPE_DEFAULT);
 			png_set_write_fn(png_ptr, this, FPngImageWrapper::user_write_compressed, FPngImageWrapper::user_flush_data);
@@ -593,6 +617,18 @@ bool FPngImageWrapper::LoadPNGHeader()
 			{
 				Format = ERGBFormat::BGRA;
 			}
+			
+			// Read all text blocks into our map
+			TextBlocks.Reset();
+			png_textp TextPointers = nullptr;
+			int32 NumTextBlocks = 0;
+			if (png_get_text(png_ptr, info_ptr, &TextPointers, &NumTextBlocks) != 0)
+			{
+				for (int32 Index = 0; Index < NumTextBlocks; Index++)
+				{
+					TextBlocks.Add(TextPointers[Index].key, TextPointers[Index].text);
+				}
+			}
 		}
 #if !PLATFORM_EXCEPTIONS_DISABLED
 		catch (const FPNGImageCRCError& e)
@@ -613,6 +649,26 @@ bool FPngImageWrapper::LoadPNGHeader()
 	return false;
 }
 
+bool FPngImageWrapper::SupportsMetadata() const
+{
+	return true;
+}
+
+void FPngImageWrapper::AddMetadata(const FString& InKey, const FString& InValue)
+{
+	TextBlocks.Add(InKey, InValue);
+}
+
+bool FPngImageWrapper::TryGetMetadata(const FString& InKey, FString& OutValue) const
+{
+	if (const FString* Value = TextBlocks.Find(InKey))
+	{
+		OutValue = *Value;
+		return true;
+	}
+
+	return false;
+}
 
 /* FPngImageWrapper static implementation
  *****************************************************************************/
diff --git a/Engine/Source/Runtime/ImageWrapper/Private/Formats/PngImageWrapper.h b/Engine/Source/Runtime/ImageWrapper/Private/Formats/PngImageWrapper.h
--- a/Engine/Source/Runtime/ImageWrapper/Private/Formats/PngImageWrapper.h
+++ b/Engine/Source/Runtime/ImageWrapper/Private/Formats/PngImageWrapper.h
@@ -53,6 +53,31 @@ public:
 	virtual bool CanSetRawFormat(const ERGBFormat InFormat, const int32 InBitDepth) const override;
 	virtual ERawImageFormat::Type GetSupportedRawFormat(const ERawImageFormat::Type InFormat) const override;
 
+	/**
+	 * Does this image type support embedded metadata in its header?
+	 *
+	 * PNG is an example of an image type which supports adding user-defined metadata to its header.
+	 */
+	virtual bool SupportsMetadata() const override;
+
+	/**
+	 * Adds a key and value to this image's metadata.  Will be saved in the image's header and restored when the image is loaded.
+	 *
+	 * @param InKey Metadata consists of key value pairs.
+	 * @param InValue Metadata consists of key value pairs.
+	 */
+	virtual void AddMetadata(const FString& InKey, const FString& InValue) override;
+
+	/**
+	 * Queries a key from this image's metadata.  If it exists, returns its corresponding value.
+	 *
+	 * @param InKey The key to locate.
+	 * @param outValue If the key exists, this is its value.
+	 *
+	 * Returns true if the key exists.
+	 */
+	virtual bool TryGetMetadata(const FString& InKey, FString& OutValue) const override;
+
 public:
 
 	/** 
@@ -99,7 +124,9 @@ private:
 	/** setjmp buffer for error recovery. */
 	jmp_buf SetjmpBuffer;
 #endif
-};
 
+	/** Metadata for/from our header. */
+	TMap<FString, FString> TextBlocks;
+};
 
 #endif	//WITH_UNREALPNG
diff --git a/Engine/Source/Runtime/ImageWrapper/Private/ImageWrapperBase.cpp b/Engine/Source/Runtime/ImageWrapper/Private/ImageWrapperBase.cpp
--- a/Engine/Source/Runtime/ImageWrapper/Private/ImageWrapperBase.cpp
+++ b/Engine/Source/Runtime/ImageWrapper/Private/ImageWrapperBase.cpp
@@ -557,3 +557,17 @@ bool IImageWrapper::GetRawImage(FDecompressedImageOutput& OutDecomressedImage)
 
 	return true;
 }
+
+bool FImageWrapperBase::SupportsMetadata() const
+{
+	return false;
+}
+
+void FImageWrapperBase::AddMetadata(const FString&, const FString&)
+{
+}
+
+bool FImageWrapperBase::TryGetMetadata(const FString&, FString&) const
+{
+	return false;
+}
\ No newline at end of file
diff --git a/Engine/Source/Runtime/ImageWrapper/Private/ImageWrapperBase.h b/Engine/Source/Runtime/ImageWrapper/Private/ImageWrapperBase.h
--- a/Engine/Source/Runtime/ImageWrapper/Private/ImageWrapperBase.h
+++ b/Engine/Source/Runtime/ImageWrapper/Private/ImageWrapperBase.h
@@ -129,6 +129,31 @@ public:
 	virtual bool SetCompressed(const void* InCompressedData, int64 InCompressedSize) override;
 	virtual bool SetRaw(const void* InRawData, int64 InRawSize, const int32 InWidth, const int32 InHeight, const ERGBFormat InFormat, const int32 InBitDepth, const int32 InBytesPerRow = 0) override;
 
+	/**
+	 * Does this image type support embedded metadata in its header?
+	 *
+	 * PNG is an example of an image type which supports adding user-defined metadata to its header.
+	 */
+	virtual bool SupportsMetadata() const override;
+
+	/**
+	 * Adds a key and value to this image's metadata.  Will be saved in the image's header and restored when the image is loaded.
+	 *
+	 * @param InKey Metadata consists of key value pairs.
+	 * @param InValue Metadata consists of key value pairs.
+	 */
+	virtual void AddMetadata(const FString& InKey, const FString& InValue) override;
+
+	/**
+	 * Queries a key from this image's metadata.  If it exists, returns its corresponding value.
+	 *
+	 * @param InKey The key to locate.
+	 * @param outValue If the key exists, this is its value.
+	 *
+	 * Returns true if the key exists.
+	 */
+	virtual bool TryGetMetadata(const FString& InKey, FString& OutValue) const override;
+
 protected:
 
 	int64 GetBytesPerPel() const { return GetRGBFormatBytesPerPel(Format,BitDepth); }
diff --git a/Engine/Source/Runtime/ImageWrapper/Public/Formats/HdrImageWrapper.h b/Engine/Source/Runtime/ImageWrapper/Public/Formats/HdrImageWrapper.h
--- a/Engine/Source/Runtime/ImageWrapper/Public/Formats/HdrImageWrapper.h
+++ b/Engine/Source/Runtime/ImageWrapper/Public/Formats/HdrImageWrapper.h
@@ -43,7 +43,33 @@ public:
 
 	IMAGEWRAPPER_API void FreeCompressedData();
 
+	/**
+	 * Does this image type support embedded metadata in its header?
+	 *
+	 * PNG is an example of an image type which supports adding user-defined metadata to its header.
+	 */
+	virtual bool SupportsMetadata() const override;
+
+	/**
+	 * Adds a key and value to this image's metadata.  Will be saved in the image's header and restored when the image is loaded.
+	 *
+	 * @param InKey Metadata consists of key value pairs.
+	 * @param InValue Metadata consists of key value pairs.
+	 */
+	virtual void AddMetadata(const FString& InKey, const FString& InValue) override;
+
+	/**
+	 * Queries a key from this image's metadata.  If it exists, returns its corresponding value.
+	 *
+	 * @param InKey The key to locate.
+	 * @param outValue If the key exists, this is its value.
+	 *
+	 * Returns true if the key exists.
+	 */
+	virtual bool TryGetMetadata(const FString& InKey, FString& OutValue) const override;
+
 	using IImageWrapper::GetRaw;
+
 private:
 	// Helpers for error exits. Set error messages and return false.
 	void SetAndLogError(const FText& InText);
diff --git a/Engine/Source/Runtime/ImageWrapper/Public/IImageWrapper.h b/Engine/Source/Runtime/ImageWrapper/Public/IImageWrapper.h
--- a/Engine/Source/Runtime/ImageWrapper/Public/IImageWrapper.h
+++ b/Engine/Source/Runtime/ImageWrapper/Public/IImageWrapper.h
@@ -433,6 +433,31 @@ public:
 		DebugImageName = InDebugImageName;
 	}
 
+	/**
+	 * Does this image type support embedded metadata in its header?
+	 *
+	 * PNG is an example of an image type which supports adding user-defined metadata to its header.
+	 */
+	virtual bool SupportsMetadata() const = 0;
+
+	/**
+	 * Adds a key and value to this image's metadata.  Will be saved in the image's header and restored when the image is loaded.
+	 *
+	 * @param InKey Metadata consists of key value pairs.
+	 * @param InValue Metadata consists of key value pairs.
+	 */
+	virtual void AddMetadata(const FString& InKey, const FString& InValue) = 0;
+
+	/**
+	 * Queries a key from this image's metadata.  If it exists, returns its corresponding value.
+	 *
+	 * @param InKey The key to locate.
+	 * @param outValue If the key exists, this is its value.
+	 *
+	 * Returns true if the key exists.
+	 */
+	virtual bool TryGetMetadata(const FString& InKey, FString& OutValue) const = 0;
+
 public:
 
 	/** Virtual destructor. */
