Fix PngImageWrapper transform logic

Turns out it's been wrong for a long time, but the libPNG 1.6 update caused it to actually break.

diff --git a/Engine/Source/Runtime/ImageWrapper/Private/Formats/PngImageWrapper.cpp b/Engine/Source/Runtime/ImageWrapper/Private/Formats/PngImageWrapper.cpp
--- a/Engine/Source/Runtime/ImageWrapper/Private/Formats/PngImageWrapper.cpp
+++ b/Engine/Source/Runtime/ImageWrapper/Private/Formats/PngImageWrapper.cpp
@@ -383,6 +383,14 @@ void FPngImageWrapper::UncompressPNGData(const ERGBFormat InFormat, const int32
 	check(InFormat == ERGBFormat::BGRA || InFormat == ERGBFormat::RGBA || InFormat == ERGBFormat::Gray);	// Other formats unsupported at present
 	check(InBitDepth == 1 || InBitDepth == 2 || InBitDepth == 4 || InBitDepth == 8 || InBitDepth == 16);	// Other formats unsupported at present
 
+	// Check if LibPNG can handle the size
+	if (Height > PNG_UINT_32_MAX / sizeof(png_bytep))
+	{
+		RawData.Empty();
+		UE_LOG(LogImageWrapper, Error, TEXT("PNG image too high to process with libPNG"));
+		return;
+	}
+
 	// Reset to the beginning of file so we can use png_read_png(), which expects to start at the beginning.
 	ReadOffset = 0;
 		
@@ -417,26 +425,17 @@ void FPngImageWrapper::UncompressPNGData(const ERGBFormat InFormat, const int32
 		// ---------------------------------------------------------------------------------------------------------
 		// Anything allocated on the stack after this point will not be destructed correctly in the case of an error
 		{
-			if (ColorType == PNG_COLOR_TYPE_PALETTE)
-			{
-				png_set_palette_to_rgb(png_ptr);
-			}
+			png_set_read_fn(png_ptr, this, FPngImageWrapper::user_read_compressed);
 
-			// @todo Oodle: really we should just call png_expand() here and remove all these conditionals
+			// Read the header
+			png_read_info(png_ptr, info_ptr);
 
-			if (((ColorType & PNG_COLOR_MASK_COLOR) == 0) && BitDepth < 8)
-			{
-				png_set_expand_gray_1_2_4_to_8(png_ptr);
-			}
+			// We want paletted data expanded out to RGB and transparencies turned into alpha
+			png_set_expand(png_ptr);
 
 			// Insert alpha channel with full opacity for RGB images without alpha
 			if ((ColorType & PNG_COLOR_MASK_ALPHA) == 0 && (InFormat == ERGBFormat::BGRA || InFormat == ERGBFormat::RGBA))
 			{
-				// png images don't set PNG_COLOR_MASK_ALPHA if they have alpha from a tRNS chunk, but png_set_add_alpha seems to be safe regardless
-				png_set_tRNS_to_alpha(png_ptr);
-
-				// note: png_set_tRNS_to_alpha is just an alias for png_expand
-
 				if (InBitDepth == 8)
 				{
 					png_set_add_alpha(png_ptr, 0xff , PNG_FILLER_AFTER);
@@ -459,22 +458,24 @@ void FPngImageWrapper::UncompressPNGData(const ERGBFormat InFormat, const int32
 			RawData.Empty(TotalBytes);
 			RawData.AddUninitialized(TotalBytes);
 
-			png_set_read_fn(png_ptr, this, FPngImageWrapper::user_read_compressed);
-
 			for (int64 i = 0; i < Height; i++)
 			{
 				row_pointers[i]= &RawData[i * BytesPerRow];
 			}
 			png_set_rows(png_ptr, info_ptr, row_pointers);
 
-			uint32 Transform = (InFormat == ERGBFormat::BGRA) ? PNG_TRANSFORM_BGR : PNG_TRANSFORM_IDENTITY;
+			// PNG files are RGBA, if we want BGRA we need to swap
+			if (InFormat == ERGBFormat::BGRA)
+			{
+				png_set_bgr(png_ptr);
+			}
 			
 			// PNG files store 16-bit pixels in network byte order (big-endian, ie. most significant bits first).
 #if PLATFORM_LITTLE_ENDIAN
 			// We're little endian so we need to swap
 			if (BitDepth == 16)
 			{
-				Transform |= PNG_TRANSFORM_SWAP_ENDIAN;
+				png_set_swap(png_ptr);
 			}
 #endif
 
@@ -486,13 +487,13 @@ void FPngImageWrapper::UncompressPNGData(const ERGBFormat InFormat, const int32
 			if ((ColorType & PNG_COLOR_MASK_COLOR) == 0 &&
 				(InFormat == ERGBFormat::RGBA || InFormat == ERGBFormat::BGRA))
 			{
-				Transform |= PNG_TRANSFORM_GRAY_TO_RGB;
+				png_set_gray_to_rgb(png_ptr);
 			}
 
 			// Convert RGB png to grayscale if requested
 			if ((ColorType & PNG_COLOR_MASK_COLOR) != 0 && InFormat == ERGBFormat::Gray)
 			{
-				png_set_rgb_to_gray_fixed(png_ptr, 2 /* warn if image is in color */, -1, -1);
+				png_set_rgb_to_gray_fixed(png_ptr, PNG_ERROR_ACTION_WARN, -1, -1);
 			}
 
 			// Strip alpha channel if requested output is grayscale
@@ -501,22 +502,29 @@ void FPngImageWrapper::UncompressPNGData(const ERGBFormat InFormat, const int32
 				// this is not necessarily the best option, instead perhaps:
 				// png_color background = {0,0,0};
 				// png_set_background(png_ptr, &background, PNG_BACKGROUND_GAMMA_SCREEN, 0, 1.0);
-				Transform |= PNG_TRANSFORM_STRIP_ALPHA;
+				png_set_strip_alpha(png_ptr);
 			}
 
 			// Reduce 16-bit to 8-bit if requested
 			if (BitDepth == 16 && InBitDepth == 8)
 			{
-				Transform |= PNG_TRANSFORM_SCALE_16;
+				png_set_scale_16(png_ptr);
 			}
 
 			// Increase 8-bit to 16-bit if requested
 			if (BitDepth <= 8 && InBitDepth == 16)
 			{
-				Transform |= PNG_TRANSFORM_EXPAND_16;
+				png_set_expand_16(png_ptr);
 			}
 
-			png_read_png(png_ptr, info_ptr, Transform, NULL);
+			// Finalize init
+			(void)png_set_interlace_handling(png_ptr);
+			png_read_update_info(png_ptr, info_ptr);
+
+			// Read the PNG
+			png_read_image(png_ptr, row_pointers);
+
+			png_read_end(png_ptr, info_ptr);
 		}
 	}
 #if !PLATFORM_EXCEPTIONS_DISABLED
