Do not enable OpenEXR on server builds to shave 3.38MB from a large server target.

diff --git a/Engine/Source/Runtime/ImageWrapper/ImageWrapper.Build.cs b/Engine/Source/Runtime/ImageWrapper/ImageWrapper.Build.cs
--- a/Engine/Source/Runtime/ImageWrapper/ImageWrapper.Build.cs
+++ b/Engine/Source/Runtime/ImageWrapper/ImageWrapper.Build.cs
@@ -58,22 +58,28 @@ public class ImageWrapper : ModuleRules
 			AddEngineThirdPartyPrivateStaticDependencies(Target, "UElibJPG");
 		}
 
+		bool bEnableEXR = false;
 		bool bUseMinimalEXR = bEnableMinimalEXR;
-		// Add openEXR lib for windows builds.
-		if ((Target.Platform.IsInGroup(UnrealPlatformGroup.Windows)) ||
-			(Target.Platform == UnrealTargetPlatform.Mac) ||
-			(Target.IsInPlatformGroup(UnrealPlatformGroup.Unix) && Target.Architecture == UnrealArch.X64))
+
+		if (Target.Type != TargetType.Server)
 		{
-			PublicDefinitions.Add("WITH_UNREALEXR=1");
-			AddEngineThirdPartyPrivateStaticDependencies(Target, "Imath");
-			AddEngineThirdPartyPrivateStaticDependencies(Target, "UEOpenExr");
-			bUseMinimalEXR = false;
+			// Add openEXR lib for desktop builds.
+			if ((Target.Platform.IsInGroup(UnrealPlatformGroup.Windows)) ||
+				(Target.Platform == UnrealTargetPlatform.Mac) ||
+				(Target.IsInPlatformGroup(UnrealPlatformGroup.Unix) && Target.Architecture == UnrealArch.X64))
+			{
+				bEnableEXR = true;
+				bUseMinimalEXR = false;
+			}
 		}
-		else
+
+		if (bEnableEXR)
 		{
-			PublicDefinitions.Add("WITH_UNREALEXR=0");
+			AddEngineThirdPartyPrivateStaticDependencies(Target, "Imath");
+			AddEngineThirdPartyPrivateStaticDependencies(Target, "UEOpenExr");
 		}
 
+		PublicDefinitions.Add("WITH_UNREALEXR=" + (bEnableEXR ? "1" : "0"));
 		PublicDefinitions.Add("WITH_UNREALEXR_MINIMAL=" + (bUseMinimalEXR ? "1" : "0"));
 
 		// Enable exceptions to allow error handling
