Changed how an image's metadata is accessed.  Previously, the image needed to remain in scope in order to access the metadata.  Now the caller caches it off and the image can be freed.

diff --git a/Engine/Source/Runtime/ImageWrapper/Private/Formats/HdrImageWrapper.cpp b/Engine/Source/Runtime/ImageWrapper/Private/Formats/HdrImageWrapper.cpp
--- a/Engine/Source/Runtime/ImageWrapper/Private/Formats/HdrImageWrapper.cpp
+++ b/Engine/Source/Runtime/ImageWrapper/Private/Formats/HdrImageWrapper.cpp
@@ -661,9 +661,9 @@ void FHdrImageWrapper::AddMetadata(const FString&, const FString&)
 {
 }
 
-bool FHdrImageWrapper::TryGetMetadata(const FString&, FString&) const
+void FHdrImageWrapper::GetMetadata(TStringMap& OutMetadata) const
 {
-	return false;
+	OutMetadata.Reset();
 }
 
 #undef LOCTEXT_NAMESPACE
diff --git a/Engine/Source/Runtime/ImageWrapper/Private/Formats/PngImageWrapper.cpp b/Engine/Source/Runtime/ImageWrapper/Private/Formats/PngImageWrapper.cpp
--- a/Engine/Source/Runtime/ImageWrapper/Private/Formats/PngImageWrapper.cpp
+++ b/Engine/Source/Runtime/ImageWrapper/Private/Formats/PngImageWrapper.cpp
@@ -667,15 +667,9 @@ void FPngImageWrapper::AddMetadata(const FString& InKey, const FString& InValue)
 	TextBlocks.Add(InKey, InValue);
 }
 
-bool FPngImageWrapper::TryGetMetadata(const FString& InKey, FString& OutValue) const
+void FPngImageWrapper::GetMetadata(TStringMap& OutMetadata) const
 {
-	if (const FString* Value = TextBlocks.Find(InKey))
-	{
-		OutValue = *Value;
-		return true;
-	}
-
-	return false;
+	OutMetadata = TextBlocks;
 }
 
 /* FPngImageWrapper static implementation
diff --git a/Engine/Source/Runtime/ImageWrapper/Private/Formats/PngImageWrapper.h b/Engine/Source/Runtime/ImageWrapper/Private/Formats/PngImageWrapper.h
--- a/Engine/Source/Runtime/ImageWrapper/Private/Formats/PngImageWrapper.h
+++ b/Engine/Source/Runtime/ImageWrapper/Private/Formats/PngImageWrapper.h
@@ -69,14 +69,11 @@ public:
 	virtual void AddMetadata(const FString& InKey, const FString& InValue) override;
 
 	/**
-	 * Queries a key from this image's metadata.  If it exists, returns its corresponding value.
+	 * Copies this image's metadata to the caller's map.
 	 *
-	 * @param InKey The key to locate.
-	 * @param outValue If the key exists, this is its value.
-	 *
-	 * Returns true if the key exists.
+	 * @param OutMetadata A copy of this image's metadata.
 	 */
-	virtual bool TryGetMetadata(const FString& InKey, FString& OutValue) const override;
+	virtual void GetMetadata(TStringMap& OutMetadata) const override;
 
 public:
 
diff --git a/Engine/Source/Runtime/ImageWrapper/Private/ImageWrapperBase.cpp b/Engine/Source/Runtime/ImageWrapper/Private/ImageWrapperBase.cpp
--- a/Engine/Source/Runtime/ImageWrapper/Private/ImageWrapperBase.cpp
+++ b/Engine/Source/Runtime/ImageWrapper/Private/ImageWrapperBase.cpp
@@ -567,7 +567,7 @@ void FImageWrapperBase::AddMetadata(const FString&, const FString&)
 {
 }
 
-bool FImageWrapperBase::TryGetMetadata(const FString&, FString&) const
+void FImageWrapperBase::GetMetadata(TStringMap& OutMetadata) const
 {
-	return false;
+	OutMetadata.Reset();
 }
\ No newline at end of file
diff --git a/Engine/Source/Runtime/ImageWrapper/Private/ImageWrapperBase.h b/Engine/Source/Runtime/ImageWrapper/Private/ImageWrapperBase.h
--- a/Engine/Source/Runtime/ImageWrapper/Private/ImageWrapperBase.h
+++ b/Engine/Source/Runtime/ImageWrapper/Private/ImageWrapperBase.h
@@ -145,14 +145,11 @@ public:
 	virtual void AddMetadata(const FString& InKey, const FString& InValue) override;
 
 	/**
-	 * Queries a key from this image's metadata.  If it exists, returns its corresponding value.
+	 * Copies this image's metadata to the caller's map.
 	 *
-	 * @param InKey The key to locate.
-	 * @param outValue If the key exists, this is its value.
-	 *
-	 * Returns true if the key exists.
+	 * @param OutMetadata A copy of this image's metadata.
 	 */
-	virtual bool TryGetMetadata(const FString& InKey, FString& OutValue) const override;
+	virtual void GetMetadata(TStringMap& OutMetadata) const override;
 
 protected:
 
diff --git a/Engine/Source/Runtime/ImageWrapper/Public/Formats/HdrImageWrapper.h b/Engine/Source/Runtime/ImageWrapper/Public/Formats/HdrImageWrapper.h
--- a/Engine/Source/Runtime/ImageWrapper/Public/Formats/HdrImageWrapper.h
+++ b/Engine/Source/Runtime/ImageWrapper/Public/Formats/HdrImageWrapper.h
@@ -59,14 +59,11 @@ public:
 	virtual void AddMetadata(const FString& InKey, const FString& InValue) override;
 
 	/**
-	 * Queries a key from this image's metadata.  If it exists, returns its corresponding value.
+	 * Copies this image's metadata to the caller's map.
 	 *
-	 * @param InKey The key to locate.
-	 * @param outValue If the key exists, this is its value.
-	 *
-	 * Returns true if the key exists.
+	 * @param OutMetadata A copy of this image's metadata.
 	 */
-	virtual bool TryGetMetadata(const FString& InKey, FString& OutValue) const override;
+	virtual void GetMetadata(TStringMap& OutMetadata) const override;
 
 	using IImageWrapper::GetRaw;
 
diff --git a/Engine/Source/Runtime/ImageWrapper/Public/IImageWrapper.h b/Engine/Source/Runtime/ImageWrapper/Public/IImageWrapper.h
--- a/Engine/Source/Runtime/ImageWrapper/Public/IImageWrapper.h
+++ b/Engine/Source/Runtime/ImageWrapper/Public/IImageWrapper.h
@@ -449,14 +449,12 @@ public:
 	virtual void AddMetadata(const FString& InKey, const FString& InValue) = 0;
 
 	/**
-	 * Queries a key from this image's metadata.  If it exists, returns its corresponding value.
+	 * Copies this image's metadata to the caller's map.
 	 *
-	 * @param InKey The key to locate.
-	 * @param outValue If the key exists, this is its value.
-	 *
-	 * Returns true if the key exists.
+	 * @param OutMetadata A copy of this image's metadata.
 	 */
-	virtual bool TryGetMetadata(const FString& InKey, FString& OutValue) const = 0;
+	using TStringMap = TMap<FString, FString>;
+	virtual void GetMetadata(TStringMap& OutMetadata) const = 0;
 
 public:
 
