Fix memory stomp on the stack in TiffImageWrapper. TIFFGetField expects a pointer to uint32 when using TIFFTAG_ROWSPERSTRIP but was passed a pointer to a uint16.

#Synced-CL 48781778

diff --git a/Engine/Source/Runtime/ImageWrapper/Private/Formats/TiffImageWrapper.cpp b/Engine/Source/Runtime/ImageWrapper/Private/Formats/TiffImageWrapper.cpp
--- a/Engine/Source/Runtime/ImageWrapper/Private/Formats/TiffImageWrapper.cpp
+++ b/Engine/Source/Runtime/ImageWrapper/Private/Formats/TiffImageWrapper.cpp
@@ -1195,7 +1195,7 @@ namespace UE::ImageWrapper::Private
 		}
 		else
 		{
-			uint16 RowPerStrip = 0;
+			uint32 RowPerStrip = 0;
 			TIFFGetField(Tiff, TIFFTAG_ROWSPERSTRIP, &RowPerStrip);
 
 			const int32 NumberOfRow = RowPerStrip > CurrSubImageHeight || RowPerStrip == 0  ? CurrSubImageHeight : RowPerStrip;
